New PlainPropsUObject plugin - add to .uproject and run with -BindPlainProps

This is an initial implementation with a PlainProps runtime that declares all native UClass/UScriptStruct/UFunction/UEnum at startup. The current simplistic approach won't work with live edit, late-loaded types, etc. The next step is UStruct and FProperty bindings.

#rnx

diff --git a/Engine/Plugins/Experimental/PlainProps/Source/Private/PlainPropsBuildSchema.cpp b/Engine/Plugins/Experimental/PlainProps/Source/Private/PlainPropsBuildSchema.cpp
--- a/Engine/Plugins/Experimental/PlainProps/Source/Private/PlainPropsBuildSchema.cpp
+++ b/Engine/Plugins/Experimental/PlainProps/Source/Private/PlainPropsBuildSchema.cpp
@@ -64,9 +64,10 @@ struct FEnumSchemaBuilder
 {
 	const FEnumDeclaration&						Declaration;
 	FEnumSchemaId								Id;
+	TOptional<ELeafWidth>						NotedWidth;
 	TSet<uint64>								NotedConstants;
 
-	void										NoteValue(uint64 Value);
+	void										NoteValue(uint64 Value, ELeafWidth Width);
 	FBuiltEnumSchema							Build() const;
 };
 
@@ -263,7 +264,7 @@ void FStructSchemaBuilder::NoteMembersRecursively(const FBuiltStruct& Struct)
 			switch (Schema.Type.GetKind())
 			{
 			case EMemberKind::Leaf:
-				AllSchemas.NoteEnum(static_cast<FEnumSchemaId>(InnerSchema)).NoteValue(Member.Value.Leaf);
+				AllSchemas.NoteEnum(static_cast<FEnumSchemaId>(InnerSchema)).NoteValue(Member.Value.Leaf, Schema.Type.AsLeaf().Width);
 				break;
 
 			case EMemberKind::Struct:
@@ -287,7 +288,7 @@ void NoteEnumValues(FEnumSchemaBuilder& Schema, const IntType* Values, uint64 Nu
 {
 	for (IntType Value : TConstArrayView64<IntType>(Values, Num))
 	{
-		Schema.NoteValue(Value);
+		Schema.NoteValue(Value, LeafWidth<sizeof(IntType)>);
 	}
 }
 
@@ -362,7 +363,7 @@ FBuiltEnumSchema FEnumSchemaBuilder::Build() const
 {
 	FBuiltEnumSchema Out = { Declaration.Type, Id };
 	Out.Mode = Declaration.Mode;
-	Out.Width = Declaration.Width;
+	Out.Width = NotedWidth.GetValue();
 
 	if (int32 Num = NotedConstants.Num())
 	{
@@ -383,8 +384,11 @@ FBuiltEnumSchema FEnumSchemaBuilder::Build() const
 	return Out;
 }
 
-void FEnumSchemaBuilder::NoteValue(uint64 Value)
+void FEnumSchemaBuilder::NoteValue(uint64 Value, ELeafWidth Width)
 {
+	check(NotedWidth == NullOpt || NotedWidth == Width);
+	NotedWidth = Width;
+
 	if (Declaration.Mode == EEnumMode::Flag)
 	{
 		if (Value == 0)
diff --git a/Engine/Plugins/Experimental/PlainProps/Source/Private/PlainPropsDeclare.cpp b/Engine/Plugins/Experimental/PlainProps/Source/Private/PlainPropsDeclare.cpp
--- a/Engine/Plugins/Experimental/PlainProps/Source/Private/PlainPropsDeclare.cpp
+++ b/Engine/Plugins/Experimental/PlainProps/Source/Private/PlainPropsDeclare.cpp
@@ -7,31 +7,84 @@
 namespace PlainProps
 {
 
+// Note: For automated upgrade purposes it could be better to not strip out enum flag aliases,
+//		 E.g. Saving E::All in E { A=1, B=2, All=A|B }, adding C=4, All=A|B|C and loading will load A|B
+//		 It's impossible to know if a user set A|B or All when saving though, we only have the value 3.
+//		 To really know, we'd need to instead save an enum oplog, i.e. {set A, set B} or {set All}
+static TConstArrayView<FEnumerator> StripAliases(TConstArrayView<FEnumerator> In, EEnumMode Mode, TArray<FEnumerator, TInlineAllocator<64>>& OutTmp)
+{
+	TBitArray<> Aliases;
+	if (Mode == EEnumMode::Flag)
+	{
+		bool bSeen0 = false;
+		uint64 Seen = 0;
+		for (FEnumerator E : In)
+		{
+			bool bAlias = E.Constant ? (Seen & E.Constant) == E.Constant : bSeen0;
+			checkf(bAlias || FMath::CountBits(E.Constant) <= 1, TEXT("Flag enums must use one bit per enumerator"));
+			
+			Aliases.Add(bAlias);
+			Seen |= E.Constant;
+			bSeen0 |= E.Constant == 0;
+		}
+	}
+	else
+	{
+		TSet<uint64, DefaultKeyFuncs<uint64>,TInlineSetAllocator<64>> Seen;
+		for (FEnumerator E : In)
+		{
+			bool bAlias;
+			Seen.FindOrAdd(E.Constant, /* out*/ &bAlias);
+			Aliases.Add(bAlias);
+		}
+	}
+
+	if (int32 NumAliases = Aliases.CountSetBits())
+	{
+		// All aliases are frequently at the end
+		int32 FirstAlias = Aliases.Find(true);
+		if (FirstAlias == In.Num() - NumAliases)
+		{
+			return In.Slice(0, FirstAlias);
+		}
+
+		// Aliases mixed in with values, make a copy and return it
+		const FEnumerator* InIt = &In[0];
+		for (bool bAlias : Aliases)
+		{
+			if (!bAlias)
+			{
+				OutTmp.Add(*InIt);
+			}
+			++InIt;
+		}
+		return OutTmp;
+	}
+
+	return In;
+}
+
 static void ValidateDeclaration(const FEnumDeclaration& Enum)
 {
 	if (Enum.Mode == EEnumMode::Flag)
 	{
 		for (FEnumerator E : Enum.GetEnumerators())
 		{
-			checkf(FMath::CountBits(E.Constant) == 1, TEXT("Flag enums must use one bit per enumerator"));
+			checkf(FMath::CountBits(E.Constant) <= 1, TEXT("Flag enums must use one bit per enumerator"));
 		}
 	}
 
 	TSet<uint32, DefaultKeyFuncs<uint32>, TInlineSetAllocator<64>> Names;
 	TSet<uint64, DefaultKeyFuncs<uint64>,TInlineSetAllocator<64>> Constants;
-	uint64 LastConstant = 0;
 	for (FEnumerator E : Enum.GetEnumerators())
 	{
-		checkf(FMath::FloorLog2_64(E.Constant) < 8 * SizeOf(Enum.Width), TEXT("Enumerator constant larger than declared width"));
+		//checkf(FMath::FloorLog2_64(E.Constant) < 8 * SizeOf(Enum.Width), TEXT("Enumerator constant larger than declared width"));
 
 		bool bDeclared;
 		Names.FindOrAdd(E.Name.Idx, /* out*/ &bDeclared);
 		checkf(!bDeclared, TEXT("Enumerator name declared twice"));
 		Constants.FindOrAdd(E.Constant, /* out*/ &bDeclared);
 		checkf(!bDeclared, TEXT("Enumerator constant declared twice"));
-
-		checkf(LastConstant <= E.Constant, TEXT("Enumerator constants must be declared in ascending order"));
-		LastConstant = E.Constant;
 	}
 }
 
@@ -71,7 +124,7 @@ void FDeclarations::DeclareStruct(FStructSchemaId DeclId, FTypeId Type, TConstAr
 	}
 }
 
-void FDeclarations::DeclareEnum(FEnumSchemaId Id, FTypeId Type, EEnumMode Mode, ELeafWidth Width, TConstArrayView<FEnumerator> Enumerators)
+void FDeclarations::DeclareEnum(FEnumSchemaId Id, FTypeId Type, EEnumMode Mode, TConstArrayView<FEnumerator> Enumerators, EEnumAliases Policy)
 {
 	if (static_cast<int32>(Id.Idx) >= DeclaredEnums.Num())
 	{
@@ -81,7 +134,13 @@ void FDeclarations::DeclareEnum(FEnumSchemaId Id, FTypeId Type, EEnumMode Mode,
 	TUniquePtr<FEnumDeclaration>& Ptr = DeclaredEnums[Id.Idx];
 	checkf(!Ptr, TEXT("'%s' is already declared"), *Debug.Print(Id));
 
-	FEnumDeclaration Header{Type, Mode, Width, IntCastChecked<uint16>(Enumerators.Num())};
+	TArray<FEnumerator, TInlineAllocator<64>> Tmp;
+	if (Policy == EEnumAliases::Strip)
+	{
+		Enumerators = StripAliases(Enumerators, Mode, /* out */ Tmp);
+	}
+
+	FEnumDeclaration Header{Type, Mode, IntCastChecked<uint16>(Enumerators.Num())};
 	void* Data = FMemory::Malloc(sizeof(FEnumDeclaration) + Enumerators.Num() * Enumerators.GetTypeSize());
 	Ptr.Reset(new (Data) FEnumDeclaration(Header));
 	CopyItems(Ptr->Enumerators, Enumerators);
diff --git a/Engine/Plugins/Experimental/PlainProps/Source/Private/PlainPropsIndex.cpp b/Engine/Plugins/Experimental/PlainProps/Source/Private/PlainPropsIndex.cpp
--- a/Engine/Plugins/Experimental/PlainProps/Source/Private/PlainPropsIndex.cpp
+++ b/Engine/Plugins/Experimental/PlainProps/Source/Private/PlainPropsIndex.cpp
@@ -124,9 +124,21 @@ FParametricTypeView FParametricTypeIndexer::Resolve(FParametricTypeId Id) const
 
 FScopeId FIdIndexerBase::NestScope(FScopeId Outer, FFlatScopeId Inner)
 {
+	check(Outer);
 	return FScopeId(NestedScopes.Index(Outer, Inner));
 }
 
+FScopeId FIdIndexerBase::NestReversedScopes(TConstArrayView<FFlatScopeId> Inners)
+{
+	FScopeId Out = NoId;
+	for (FFlatScopeId Inner : ReverseIterate(Inners))
+	{
+		Out = Out ? NestScope(Out, Inner) : FScopeId(Inner);
+	}
+
+	return Out;
+}
+
 FParametricTypeId FIdIndexerBase::MakeParametricTypeId(FOptionalConcreteTypenameId Name, TConstArrayView<FTypeId> Params)
 {
 	return ParametricTypes.Index({Name, IntCastChecked<uint8>(Params.Num()), Params.GetData()});
diff --git a/Engine/Plugins/Experimental/PlainProps/Source/Public/PlainPropsBind.h b/Engine/Plugins/Experimental/PlainProps/Source/Public/PlainPropsBind.h
--- a/Engine/Plugins/Experimental/PlainProps/Source/Public/PlainPropsBind.h
+++ b/Engine/Plugins/Experimental/PlainProps/Source/Public/PlainPropsBind.h
@@ -321,9 +321,15 @@ private:
 template<typename T>
 struct TCustomBind { using Type = void; };
 
+template<class T>
+struct TCustomDeltaBind { using Type = void; };
+
 template<typename T>
 using CustomBind = typename TCustomBind<T>::Type;
 
+template<typename T>
+using CustomDeltaBind = typename TCustomDeltaBind<T>::Type;
+
 ////////////////////////////////////////////////////////////////////////////////////////////////
 
 template<typename T>
@@ -580,7 +586,7 @@ public:
 	PLAINPROPS_API const FSchemaBinding&	GetStruct(FStructSchemaId BindId) const;
 	PLAINPROPS_API void						DropStruct(FStructSchemaId BindId);
 
-	virtual FStructSchemaId					GetDeclId(FStructSchemaId BindId) const override final;
+	PLAINPROPS_API virtual FStructSchemaId	GetDeclId(FStructSchemaId BindId) const override final;
 private:
 	TArray<TUniquePtr<FSchemaBinding>>		Bindings;
 	const FDebugIds&						Debug;
@@ -1034,7 +1040,7 @@ FEnumSchemaId DeclareNativeEnum(FDeclarations& Out, EEnumMode Mode)
 		Enumerator.Name = Ids::IndexName(Ctti::Enumerators[&Enumerator - Enumerators].Name);
 		Enumerator.Constant = static_cast<uint64>(static_cast<UnderlyingType>(Ctti::Enumerators[&Enumerator - Enumerators].Constant));
 	}
-	Out.DeclareEnum(Id, Type, Mode, LeafWidth<sizeof(UnderlyingType)>, Enumerators);
+	Out.DeclareEnum(Id, Type, Mode, Enumerators, EEnumAliases::Fail);
 
 	return Id;
 }
diff --git a/Engine/Plugins/Experimental/PlainProps/Source/Public/PlainPropsDeclare.h b/Engine/Plugins/Experimental/PlainProps/Source/Public/PlainPropsDeclare.h
--- a/Engine/Plugins/Experimental/PlainProps/Source/Public/PlainPropsDeclare.h
+++ b/Engine/Plugins/Experimental/PlainProps/Source/Public/PlainPropsDeclare.h
@@ -22,7 +22,7 @@ struct FEnumDeclaration
 {
 	FTypeId					Type;			// Could be removed
 	EEnumMode				Mode;
-	ELeafWidth				Width;
+	//EOptionalLeafWidth	Width;			// Possible approach to declare enums w/o any noted values
 	uint16					NumEnumerators;
 	FEnumerator				Enumerators[0];	// Constants must be unique, no aliases allowed
 
@@ -50,13 +50,17 @@ struct FStructDeclaration
 	}
 };
 
+// Enum values are stored as integers. Aliased enum values are illegal, including composite flags.
+// Aliases can be automatically removed on declaration or detected and fail hard.
+enum class EEnumAliases { Strip, Fail };
+
 class FDeclarations
 {
 public:
 	UE_NONCOPYABLE(FDeclarations);
 	explicit FDeclarations(const FDebugIds& In) : Debug(In) {}
 
-	PLAINPROPS_API void								DeclareEnum(FEnumSchemaId DeclId, FTypeId Type, EEnumMode Mode, ELeafWidth Width, TConstArrayView<FEnumerator> Enumerators);
+	PLAINPROPS_API void								DeclareEnum(FEnumSchemaId DeclId, FTypeId Type, EEnumMode Mode, TConstArrayView<FEnumerator> Enumerators, EEnumAliases Policy);
 	// Declare struct with ref count 1 or increment it and check that previous declaration matches
 	PLAINPROPS_API void								DeclareStruct(FStructSchemaId DeclId, FTypeId Type, TConstArrayView<FMemberId> MemberOrder, EMemberPresence Occupancy, FOptionalStructSchemaId Super = {});
 	
diff --git a/Engine/Plugins/Experimental/PlainProps/Source/Public/PlainPropsIndex.h b/Engine/Plugins/Experimental/PlainProps/Source/Public/PlainPropsIndex.h
--- a/Engine/Plugins/Experimental/PlainProps/Source/Public/PlainPropsIndex.h
+++ b/Engine/Plugins/Experimental/PlainProps/Source/Public/PlainPropsIndex.h
@@ -50,6 +50,7 @@ class FIdIndexerBase : public FDebugIds
 {
 public:
 	PLAINPROPS_API FScopeId				NestScope(FScopeId Outer, FFlatScopeId Inner);
+	PLAINPROPS_API FScopeId				NestReversedScopes(TConstArrayView<FFlatScopeId> Inners);
 	PLAINPROPS_API FParametricTypeId	MakeParametricTypeId(FOptionalConcreteTypenameId Name, TConstArrayView<FTypeId> Params);
 	PLAINPROPS_API FTypeId				MakeParametricType(FTypeId Type, TConstArrayView<FTypeId> Params);
 	PLAINPROPS_API FTypeId				MakeAnonymousParametricType(TConstArrayView<FTypeId> Params);
@@ -88,10 +89,10 @@ class TIdIndexer : public FIdIndexerBase
 public:
 	template<typename T> FNameId		MakeName(T&& Str)									{ return { IntCastChecked<uint32>(Names.Add(NameType(Str)).AsInteger()) }; }
 	template<typename T> FMemberId		NameMember(T&& Name)								{ return { MakeName(Name) }; }
-	template<typename T> FScopeId		MakeScope(T&& Str)									{ return FScopeId(static_cast<FFlatScopeId>(MakeName(Str))); }
+	template<typename T> FFlatScopeId	MakeScope(T&& Str)									{ return static_cast<FFlatScopeId>(MakeName(Str)); }
 	template<typename T> FScopeId		NestScope(FScopeId Outer, T&& Inner)				{ return FIdIndexerBase::NestScope(Outer, static_cast<FFlatScopeId>(MakeName(Inner))); }
 	template<typename T> FTypenameId	MakeTypename(T&& Name)								{ return FTypenameId(FConcreteTypenameId(MakeName(Name))); }
-	template<typename T> FTypeId		MakeType(T&& Scope, T&& Name)						{ return { MakeScope(Scope), MakeTypename(Name) }; }
+	template<typename T> FTypeId		MakeType(T&& Scope, T&& Name)						{ return { FScopeId(MakeScope(Scope)), MakeTypename(Name) }; }
 
 	NameType							ResolveName(FNameId Id)	const						{ return Names.Get(FSetElementId::FromInteger(Id.Idx)); }
 	virtual uint32						NumNames() const override final						{ return static_cast<uint32>(Names.Num()); }
diff --git a/Engine/Plugins/Experimental/PlainProps/Source/Public/PlainPropsTypes.h b/Engine/Plugins/Experimental/PlainProps/Source/Public/PlainPropsTypes.h
--- a/Engine/Plugins/Experimental/PlainProps/Source/Public/PlainPropsTypes.h
+++ b/Engine/Plugins/Experimental/PlainProps/Source/Public/PlainPropsTypes.h
@@ -146,8 +146,8 @@ class FScopeId
 	uint32 Handle;
 public:
 	FScopeId(FNoId) : Handle(Unscoped) {}
-	explicit FScopeId(FFlatScopeId Flat) : Handle(Flat.Name.Idx) 				{ check(AsFlat().Name == Flat.Name); }
-	explicit FScopeId(FNestedScopeId Nested) : Handle(Nested.Idx | NestedBit) 	{ check(AsNested().Idx == Nested.Idx); }
+	FScopeId(FFlatScopeId Flat) : Handle(Flat.Name.Idx) 				{ check(AsFlat().Name == Flat.Name); }
+	FScopeId(FNestedScopeId Nested) : Handle(Nested.Idx | NestedBit) 	{ check(AsNested().Idx == Nested.Idx); }
 
 	explicit						operator bool() const			{ return Handle != ~0u; }
 	bool							IsFlat() const 					{ return !(Handle & NestedBit); }
diff --git a/Engine/Plugins/Experimental/PlainProps/Source/Public/PlainPropsUeCoreBindings.h b/Engine/Plugins/Experimental/PlainProps/Source/Public/PlainPropsUeCoreBindings.h
--- a/Engine/Plugins/Experimental/PlainProps/Source/Public/PlainPropsUeCoreBindings.h
+++ b/Engine/Plugins/Experimental/PlainProps/Source/Public/PlainPropsUeCoreBindings.h
@@ -980,12 +980,24 @@ struct TRangeBind<TSet<T, KeyFuncs, SetAllocator>>
 	using Type = UE::TSetBinding<T, KeyFuncs, SetAllocator>;
 };
 
+template <typename T, typename KeyFuncs, typename SetAllocator>
+struct TCustomDeltaBind<TSet<T, KeyFuncs, SetAllocator>>
+{
+	using Type = UE::TSetDeltaBinding<T, KeyFuncs, SetAllocator>;
+};
+
 template <typename K, typename V, typename SetAllocator, typename KeyFuncs>
 struct TRangeBind<TMap<K, V, SetAllocator, KeyFuncs>>
 {
 	using Type = UE::TMapBinding<K, V, SetAllocator, KeyFuncs>;
 };
 
+template <typename K, typename V, typename SetAllocator, typename KeyFuncs>
+struct TCustomDeltaBind<TMap<K, V, SetAllocator, KeyFuncs>>
+{
+	using Type = UE::TMapDeltaBinding<K, V, SetAllocator, KeyFuncs>;
+};
+
 template<typename T>
 struct TRangeBind<TOptional<T>>
 {
diff --git a/Engine/Plugins/Experimental/PlainProps/Source/Tests/PlainPropsTests.cpp b/Engine/Plugins/Experimental/PlainProps/Source/Tests/PlainPropsTests.cpp
--- a/Engine/Plugins/Experimental/PlainProps/Source/Tests/PlainPropsTests.cpp
+++ b/Engine/Plugins/Experimental/PlainProps/Source/Tests/PlainPropsTests.cpp
@@ -171,10 +171,10 @@ FStructSchemaId FTestBatchBuilder::DeclareStruct(const char* Scope, const char*
 	return DeclareStruct(MakeType(Scope, Name), MemberOrder, Occupancy, Super);
 }
 
-FEnumSchemaId FTestBatchBuilder::DeclareEnum(FTypeId Type, EEnumMode Mode, ELeafWidth Width, std::initializer_list<const char*> InNames, std::initializer_list<uint64> Constants)
+FEnumSchemaId FTestBatchBuilder::DeclareEnum(FTypeId Type, EEnumMode Mode, ELeafWidth, std::initializer_list<const char*> InNames, std::initializer_list<uint64> Constants)
 {
 	FEnumSchemaId Id = IndexEnum(Type);
-	Declarations.DeclareEnum(Id, Type, Mode, Width, MakeEnumerators(InNames, Constants));
+	Declarations.DeclareEnum(Id, Type, Mode, MakeEnumerators(InNames, Constants), EEnumAliases::Fail);
 	return Id;
 }
 
diff --git a/Engine/Plugins/Experimental/PlainProps/Source/Tests/PlainPropsUeCoreTests.cpp b/Engine/Plugins/Experimental/PlainProps/Source/Tests/PlainPropsUeCoreTests.cpp
--- a/Engine/Plugins/Experimental/PlainProps/Source/Tests/PlainPropsUeCoreTests.cpp
+++ b/Engine/Plugins/Experimental/PlainProps/Source/Tests/PlainPropsUeCoreTests.cpp
@@ -33,7 +33,7 @@ struct FIds
 	static FNameId			IndexName(FAnsiStringView Name)				{ return GNames.MakeName(FName(Name)); }
 	static FMemberId		IndexMember(FAnsiStringView Name)			{ return GNames.NameMember(FName(Name)); }
 	static FTypenameId		IndexTypename(FAnsiStringView Name)			{ return GNames.MakeTypename(FName(Name)); }
-	static FScopeId			IndexScope(FAnsiStringView Name)			{ return GNames.MakeScope(FName(Name)); }
+	static FFlatScopeId		IndexScope(FAnsiStringView Name)			{ return GNames.MakeScope(FName(Name)); }
 	static FEnumSchemaId	IndexEnum(FTypeId Type)						{ return GNames.IndexEnum(Type); }
 	static FStructSchemaId	IndexStruct(FTypeId Type)					{ return GNames.IndexStruct(Type); }
 	static FIdIndexerBase&	GetIndexer()								{ return GNames; }
@@ -52,27 +52,8 @@ struct FDefaultRuntime
 	static FCustomBindings&			GetCustoms()		{ return GCustoms; }
 };
 
-
 ////////////////////////////////////////////////////////////////////////////////////////////////
 
-template<class T>
-struct TCustomDeltaBind
-{
-	using Type = void;
-};
-
-template <typename T, typename KeyFuncs, typename SetAllocator>
-struct TCustomDeltaBind<TSet<T, KeyFuncs, SetAllocator>>
-{
-	using Type = UE::TSetDeltaBinding<T, KeyFuncs, SetAllocator>;
-};
-
-template <typename K, typename V, typename SetAllocator, typename KeyFuncs>
-struct TCustomDeltaBind<TMap<K, V, SetAllocator, KeyFuncs>>
-{
-	using Type = UE::TMapDeltaBinding<K, V, SetAllocator, KeyFuncs>;
-};
-
 struct FDeltaRuntime : FDefaultRuntime
 {
 	template<class T> using CustomBindings = TCustomDeltaBind<T>;
diff --git a/Engine/Plugins/Experimental/PlainPropsUObject/PlainPropsUObject.uplugin b/Engine/Plugins/Experimental/PlainPropsUObject/PlainPropsUObject.uplugin
new file mode 100644
--- /dev/null
+++ b/Engine/Plugins/Experimental/PlainPropsUObject/PlainPropsUObject.uplugin
@@ -0,0 +1,26 @@
+{
+	"FileVersion": 3,
+	"FriendlyName": "PlainPropsUObject",
+	"Description": "New Serialization Stack Prototype - CoreUObject Bindings",
+	"Category": "Experimental",
+	"CreatedBy": "Epic Games, Inc.",
+	"CreatedByURL": "https://epicgames.com",
+	"EnabledByDefault": false,
+	"IsExperimentalVersion": true,
+	"SupportedTargetPlatforms": [
+		"Win64"
+	],
+	"Modules": [
+		{
+			"Name": "PlainPropsUObject",
+			"Type": "Runtime",
+			"PlatformAllowList": [ "Win64" ]
+		}
+	],
+	"Plugins": [
+		{
+			"Name": "PlainProps",
+			"Enabled": true
+		}
+	]
+}
\ No newline at end of file
diff --git a/Engine/Plugins/Experimental/PlainPropsUObject/Source/PlainPropsUObject.Build.cs b/Engine/Plugins/Experimental/PlainPropsUObject/Source/PlainPropsUObject.Build.cs
new file mode 100644
--- /dev/null
+++ b/Engine/Plugins/Experimental/PlainPropsUObject/Source/PlainPropsUObject.Build.cs
@@ -0,0 +1,20 @@
+// Copyright Epic Games, Inc. All Rights Reserved.
+
+using UnrealBuildTool;
+
+public class PlainPropsUObject : ModuleRules
+{
+	public PlainPropsUObject(ReadOnlyTargetRules Target) : base(Target)
+	{
+		bDisableStaticAnalysis = true;
+
+		CppStandard = CppStandardVersion.Cpp20;
+
+		PrivateDependencyModuleNames.AddRange(
+			new string[] {
+				"Core",
+				"CoreUObject",
+				"PlainProps",
+		});
+	}
+}
diff --git a/Engine/Plugins/Experimental/PlainPropsUObject/Source/Private/PlainPropsUObjectBindings.cpp b/Engine/Plugins/Experimental/PlainPropsUObject/Source/Private/PlainPropsUObjectBindings.cpp
new file mode 100644
--- /dev/null
+++ b/Engine/Plugins/Experimental/PlainPropsUObject/Source/Private/PlainPropsUObjectBindings.cpp
@@ -0,0 +1,143 @@
+// Copyright Epic Games, Inc. All Rights Reserved.
+
+#include "PlainPropsUObjectRuntime.h"
+#include "Misc/CommandLine.h"
+#include "Misc/CoreDelegates.h"
+#include "UObject/UnrealType.h"
+#include "UObject/UObjectIterator.h"
+
+
+namespace PlainProps::UE
+{
+
+FGlobals::FGlobals() 
+: Types(Names)
+, Schemas(Names)
+, Customs(Names)
+, DeltaCustoms(Names, &Customs)
+{}
+
+FGlobals GUE;
+
+////////////////////////////////////////////////////////////////////////////////////////////////
+
+static FTypeId IndexType(UField* Field)
+{	
+	FTypenameId Name = GUE.Names.MakeTypename(Field->GetFName());
+	
+	// This wouldn't be needed in an intrusive or cached solution
+	TArray<FFlatScopeId, TInlineAllocator<64>> ReversedOuters;
+	for (const UObject* Outer = Field->GetOuter(); Outer; Outer = Outer->GetOuter())
+	{
+		ReversedOuters.Add(GUE.Names.MakeScope(Outer->GetFName()));
+	}
+	
+	return { GUE.Names.NestReversedScopes(ReversedOuters), Name };
+}
+
+static FOptionalStructSchemaId IndexSuper(UStruct* Struct)
+{
+	return Struct ? ToOptional(GUE.Names.IndexStruct(IndexType(Struct))) : NoId;
+}
+
+static EMemberPresence GetOccupancy(UStruct* Struct)
+{
+	if (Struct->HasAnyCastFlags(CASTCLASS_UScriptStruct))
+	{
+		EStructFlags Flags = static_cast<UScriptStruct*>(Struct)->StructFlags;
+		return (Flags & (STRUCT_Immutable | STRUCT_Atomic)) ? EMemberPresence::RequireAll : EMemberPresence::AllowSparse; 
+	}
+	return EMemberPresence::AllowSparse;
+}
+
+static FStructSchemaId DeclareStruct(UStruct* Struct)
+{
+	FTypeId Type = IndexType(Struct);
+	FStructSchemaId Id = GUE.Names.IndexStruct(Type);
+	FOptionalStructSchemaId  Super = IndexSuper(Struct->GetInheritanceSuper());
+	EMemberPresence Occupancy = GetOccupancy(Struct);
+	TArray<FMemberId, TInlineAllocator<64>> Members;
+	for (FProperty* Property = Struct->PropertyLink; Property; Property = Property->PropertyLinkNext)
+	{
+		Members.Add(GUE.Names.NameMember(Property->GetFName()));
+	}
+	
+	GUE.Types.DeclareStruct(Id, Type, Members, Occupancy, Super);
+
+	return Id;
+}
+
+////////////////////////////////////////////////////////////////////////////////////////////////
+
+static bool IsMax(FName EnumName, FName ValueName)
+{
+	if (ValueName.GetNumber() != NAME_NO_NUMBER_INTERNAL)
+	{
+		return false;
+	}
+
+	// All type names and enum constants are ASCII
+	ANSICHAR Buffer[NAME_SIZE];
+	ValueName.GetComparisonNameEntry()->GetAnsiName(Buffer);
+	FAnsiStringView Value(Buffer);
+	if (Value.Len() < 5 || !Value.EndsWith("MAX"))
+	{
+		return false;
+	}
+
+	// Todo: Check EnumName too, maybe based on ECppForm
+	char PreMaxChar = Value[Value.Len() - 4];
+	return PreMaxChar == ':' || '_';
+}
+
+static FEnumSchemaId DeclareEnum(UEnum* Enum)
+{
+	FTypeId Type = IndexType(Enum);
+	FEnumSchemaId Id = GUE.Names.IndexEnum(Type);
+	EEnumMode Mode = Enum->HasAnyEnumFlags(EEnumFlags::Flags) ? EEnumMode::Flag : EEnumMode::Flat;
+	TArray<FEnumerator, TInlineAllocator<64>> Enumerators;
+	for (int32 Idx = 0, Num = Enum->NumEnums(); Idx < Num; ++Idx)
+	{
+		FName ValueName = Enum->GetNameByIndex(Idx);
+		if (Idx < Num - 1 || !IsMax(Enum->GetFName(), ValueName))
+		{
+			Enumerators.Emplace(GUE.Names.MakeName(ValueName), static_cast<uint64>(Enum->GetValueByIndex(Idx)));
+		}
+	}
+	
+#if WITH_METADATA
+	// IsMax() classifies more names as "max" than Enum->ContainsExistingMax()
+	check(Enumerators.Num() == Enum->NumEnums() || Enum->ContainsExistingMax() || Enum->HasMetaData(TEXT("Hidden"), Enumerators.Num()));
+#endif
+
+	GUE.Types.DeclareEnum(Id, Type, Mode, Enumerators, EEnumAliases::Strip);
+	return Id;
+}
+
+////////////////////////////////////////////////////////////////////////////////////////////////
+
+static void BindInitialTypes()
+{
+	if (!FParse::Param(FCommandLine::Get(), TEXT("BindPlainProps")))
+	{
+		return;
+	}
+
+	UE_LOG(LogInit, Display, TEXT("Binding classes to PlainProps schemas"));
+	for (TObjectIterator<UField> It; It; ++It)
+	{
+		UObject* Object = *It; 
+		if (UStruct* Struct = Cast<UStruct>(Object))
+		{
+			DeclareStruct(Struct);
+		}
+		else if (UEnum* Enum = Cast<UEnum>(Object))
+		{
+			DeclareEnum(Enum);
+		}	
+	}
+}
+
+static FDelegateHandle PreMainCallback = FCoreDelegates::OnCommandletPreMain.AddStatic(&BindInitialTypes);
+
+} // namespace PlainProps::UE
\ No newline at end of file
diff --git a/Engine/Plugins/Experimental/PlainPropsUObject/Source/Private/PlainPropsUObjectModule.cpp b/Engine/Plugins/Experimental/PlainPropsUObject/Source/Private/PlainPropsUObjectModule.cpp
new file mode 100644
--- /dev/null
+++ b/Engine/Plugins/Experimental/PlainPropsUObject/Source/Private/PlainPropsUObjectModule.cpp
@@ -0,0 +1,6 @@
+// Copyright Epic Games, Inc. All Rights Reserved.
+
+#include "PlainPropsUObjectModule.h"
+#include "Modules/ModuleManager.h"
+
+IMPLEMENT_MODULE(FPlainPropsUObjectModule, PlainPropsUObject);
\ No newline at end of file
diff --git a/Engine/Plugins/Experimental/PlainPropsUObject/Source/Public/PlainPropsUObjectModule.h b/Engine/Plugins/Experimental/PlainPropsUObject/Source/Public/PlainPropsUObjectModule.h
new file mode 100644
--- /dev/null
+++ b/Engine/Plugins/Experimental/PlainPropsUObject/Source/Public/PlainPropsUObjectModule.h
@@ -0,0 +1,7 @@
+// Copyright Epic Games, Inc. All Rights Reserved.
+
+#pragma once
+
+#include "Modules/ModuleInterface.h"
+
+class FPlainPropsUObjectModule final : public IModuleInterface {};
\ No newline at end of file
diff --git a/Engine/Plugins/Experimental/PlainPropsUObject/Source/Public/PlainPropsUObjectRuntime.h b/Engine/Plugins/Experimental/PlainPropsUObject/Source/Public/PlainPropsUObjectRuntime.h
new file mode 100644
--- /dev/null
+++ b/Engine/Plugins/Experimental/PlainPropsUObject/Source/Public/PlainPropsUObjectRuntime.h
@@ -0,0 +1,60 @@
+// Copyright Epic Games, Inc. All Rights Reserved.
+
+#pragma once
+
+#include "PlainPropsBind.h"
+#include "PlainPropsDeclare.h"
+#include "PlainPropsIndex.h"
+#include "PlainPropsTypes.h"
+#include "UObject/NameTypes.h"
+
+namespace PlainProps
+{
+template<> inline void AppendString(FString& Out, const FName& Name) { Name.AppendString(Out); }
+}
+
+namespace PlainProps::UE
+{
+
+struct FGlobals
+{
+	FGlobals();
+
+	TIdIndexer<FName>	Names;
+	FDeclarations		Types;
+	FSchemaBindings		Schemas;
+	FCustomBindings		Customs;
+	FCustomBindings		DeltaCustoms;
+};
+PLAINPROPSUOBJECT_API extern FGlobals GUE;
+
+struct FIds
+{
+	static FNameId				IndexName(FAnsiStringView Name)			{ return GUE.Names.MakeName(FName(Name)); }
+	static FMemberId			IndexMember(FAnsiStringView Name)		{ return GUE.Names.NameMember(FName(Name)); }
+	static FTypenameId			IndexTypename(FAnsiStringView Name)		{ return GUE.Names.MakeTypename(FName(Name)); }
+	static FFlatScopeId			IndexScope(FAnsiStringView Name)		{ return GUE.Names.MakeScope(FName(Name)); }
+	static FEnumSchemaId		IndexEnum(FTypeId Type)					{ return GUE.Names.IndexEnum(Type); }
+	static FStructSchemaId		IndexStruct(FTypeId Type)				{ return GUE.Names.IndexStruct(Type); }
+	static FIdIndexerBase&		GetIndexer()							{ return GUE.Names; }
+	static const FDebugIds&		GetDebug()								{ return GUE.Names; }
+};
+
+struct FDefaultRuntime
+{
+	using Ids = FIds;
+	template<class T> using CustomBindings = TCustomBind<T>;
+
+	static FDeclarations&		GetTypes()			{ return GUE.Types; }
+	static FSchemaBindings&		GetSchemas()		{ return GUE.Schemas; }
+	static FCustomBindings&		GetCustoms()		{ return GUE.Customs; }
+};
+
+struct FDeltaRuntime : FDefaultRuntime
+{
+	template<class T> using CustomBindings = TCustomDeltaBind<T>;
+
+	static FCustomBindings&		GetCustoms()		{ return GUE.DeltaCustoms; }
+};
+
+} // namespace PlainProps::UE
\ No newline at end of file
