PlainProps FProperty binding dispatch and FEnumProperty

#rnx

diff --git a/Engine/Plugins/Experimental/PlainProps/Source/Public/PlainPropsTypes.h b/Engine/Plugins/Experimental/PlainProps/Source/Public/PlainPropsTypes.h
--- a/Engine/Plugins/Experimental/PlainProps/Source/Public/PlainPropsTypes.h
+++ b/Engine/Plugins/Experimental/PlainProps/Source/Public/PlainPropsTypes.h
@@ -20,6 +20,12 @@ inline constexpr SIZE_T SizeOf(ELeafWidth Width)
 	return SIZE_T(1) << static_cast<uint32>(Width);
 }
 
+inline ELeafWidth WidthOf(SIZE_T Size)
+{
+	check(Size == 1 || Size == 2 || Size == 4 || Size == 8);
+	return static_cast<ELeafWidth>(FMath::FloorLog2NonZero_64(Size));
+}
+
 inline SIZE_T SizeOf(ERangeSizeType Width)
 {
 	check(Width != ERangeSizeType::Uni);
diff --git a/Engine/Plugins/Experimental/PlainPropsUObject/Source/Private/PlainPropsUObjectBindings.cpp b/Engine/Plugins/Experimental/PlainPropsUObject/Source/Private/PlainPropsUObjectBindings.cpp
--- a/Engine/Plugins/Experimental/PlainPropsUObject/Source/Private/PlainPropsUObjectBindings.cpp
+++ b/Engine/Plugins/Experimental/PlainPropsUObject/Source/Private/PlainPropsUObjectBindings.cpp
@@ -1,10 +1,19 @@
 // Copyright Epic Games, Inc. All Rights Reserved.
 
 #include "PlainPropsUObjectRuntime.h"
+#include "Containers/PagedArray.h"
 #include "Misc/CommandLine.h"
 #include "Misc/CoreDelegates.h"
+#include "UObject/AnsiStrProperty.h"
+#include "UObject/EnumProperty.h"
+#include "UObject/FieldPathProperty.h"
+#include "UObject/PropertyOptional.h"
+#include "UObject/StrProperty.h"
+#include "UObject/TextProperty.h"
+#include "UObject/Utf8StrProperty.h"
 #include "UObject/UnrealType.h"
 #include "UObject/UObjectIterator.h"
+#include "UObject/VerseValueProperty.h"
 
 
 namespace PlainProps::UE
@@ -116,6 +125,170 @@ static FEnumSchemaId DeclareEnum(UEnum* Enum)
 
 ////////////////////////////////////////////////////////////////////////////////////////////////
 
+using FRangeBindingArray = TPagedArray<FRangeBinding, 1024>;
+
+inline FMemberBinding Todo(FProperty* Property)
+{
+	return FMemberBinding(Property->GetOffset_ForInternal());
+}
+
+inline FMemberBinding BindImpl(FEnumProperty* Property, FMemberBindType	InnermostType, FOptionalSchemaId InnermostSchema, TConstArrayView<FRangeBinding> RangeBindings)
+{
+	FMemberBinding Out(Property->GetOffset_ForInternal());
+	Out.InnermostType = InnermostType;
+	Out.InnermostSchema = InnermostSchema;
+	Out.RangeBindings = RangeBindings;
+
+	if (Property->GetElementSize() == Property->GetSize())
+	{
+		return Out;
+	}
+	
+	// If we hit this limit, generate one-level of nested structs
+	uint16 ArrayDim = IntCastChecked<uint16>(Property->GetSize() / Property->GetElementSize());
+
+	// todo
+	return Out;
+}
+
+
+inline FMemberBinding BindNumeric(FNumericProperty* Property, uint64 CastFlags) { return Todo(Property); }
+inline FMemberBinding BindEnum(FEnumProperty* Property)
+{
+	FSchemaId Id = GUE.Names.IndexEnum(IndexType(Property->GetEnum()));
+	FUnpackedLeafType Leaf = {ELeafType::Enum, WidthOf(Property->GetElementSize())};
+	return BindImpl(Property, FMemberBindType(Leaf), Id, {});
+}
+inline FMemberBinding BindBool(FBoolProperty* Property) { return Todo(Property); }
+
+inline FMemberBinding BindReference(FObjectPropertyBase* Property, uint64 CastFlags) { return Todo(Property); }
+inline FMemberBinding BindInterface(FInterfaceProperty* Property) { return Todo(Property); }
+
+inline FMemberBinding BindStruct(FStructProperty* Property, FRangeBindingArray& Ranges)	{ return Todo(Property); }
+inline FMemberBinding BindDelegate(FDelegateProperty* Property, FRangeBindingArray& Ranges) { return Todo(Property); }
+inline FMemberBinding BindMulticastDelegate(FMulticastDelegateProperty* Property, FRangeBindingArray& Ranges) { return Todo(Property); }
+
+inline FMemberBinding BindArray(FArrayProperty* Property, FRangeBindingArray& Ranges) { return Todo(Property); }
+inline FMemberBinding BindMap(FMapProperty* Property, FRangeBindingArray& Ranges) { return Todo(Property); }
+inline FMemberBinding BindSet(FSetProperty* Property, FRangeBindingArray& Ranges) { return Todo(Property); }
+inline FMemberBinding BindOptional(FOptionalProperty* Property, FRangeBindingArray& Ranges) { return Todo(Property); }
+
+template<typename StrProperty>
+inline FMemberBinding BindStr(StrProperty* Property, FRangeBindingArray& Ranges) { return Todo(Property); }
+inline FMemberBinding BindName(FNameProperty* Property) { return Todo(Property); }
+inline FMemberBinding BindText(FTextProperty* Property, FRangeBindingArray& Ranges) { return Todo(Property); }
+
+inline FMemberBinding BindFieldPath(FFieldPathProperty* Property) { return Todo(Property); }
+
+#if WITH_VERSE_VM
+inline FMemberBinding BindVValue(FVRestValueProperty* Property, FRangeBindingArray& Ranges);
+inline FMemberBinding BindVRestValue(FVRestValueProperty* Property, FRangeBindingArray& Ranges);
+#endif
+
+template<uint64 Mask>
+inline bool HasAny(uint64 Flags)
+{
+	return (Mask & Flags) != 0;
+}
+
+static FMemberBinding BindProperty(FProperty* Property, FRangeBindingArray& Ranges)
+{
+	constexpr uint64 LeafMask = CASTCLASS_FNumericProperty | CASTCLASS_FEnumProperty| CASTCLASS_FBoolProperty;
+	constexpr uint64 ReferenceMask = CASTCLASS_FObjectPropertyBase | CASTCLASS_FInterfaceProperty;
+	constexpr uint64 ContainerMask = CASTCLASS_FArrayProperty | CASTCLASS_FSetProperty| CASTCLASS_FMapProperty  | CASTCLASS_FOptionalProperty;
+	constexpr uint64 StructuralMask = CASTCLASS_FStructProperty | CASTCLASS_FDelegateProperty | CASTCLASS_FMulticastDelegateProperty;
+	constexpr uint64 StringMask = CASTCLASS_FStrProperty | CASTCLASS_FUtf8StrProperty | CASTCLASS_FAnsiStrProperty | CASTCLASS_FNameProperty | CASTCLASS_FTextProperty;
+
+	uint64 Flags = Property->GetCastFlags();
+	if (HasAny<LeafMask>(Flags))
+	{
+		if (HasAny<CASTCLASS_FNumericProperty>(Flags))
+		{
+			return BindNumeric(static_cast<FNumericProperty*>(Property), Flags);
+		}
+		return HasAny<CASTCLASS_FEnumProperty>(Flags)
+			? BindEnum(static_cast<FEnumProperty*>(Property))
+			: BindBool(static_cast<FBoolProperty*>(Property));
+	}
+	else if (HasAny<ReferenceMask>(Flags))
+	{
+		return HasAny<CASTCLASS_FObjectPropertyBase>(Flags)
+			? BindReference(static_cast<FObjectPropertyBase*>(Property), Flags)
+			: BindInterface(static_cast<FInterfaceProperty*>(Property));
+	}
+	else if (HasAny<StructuralMask>(Flags))
+	{
+		if (HasAny<CASTCLASS_FStructProperty>(Flags))
+		{
+			return BindStruct(static_cast<FStructProperty*>(Property), Ranges);
+		}
+		return HasAny<CASTCLASS_FDelegateProperty>(Flags)
+			? BindDelegate(static_cast<FDelegateProperty*>(Property), Ranges)
+			: BindMulticastDelegate(static_cast<FMulticastDelegateProperty*>(Property), Ranges);
+	}
+	else if (HasAny<ContainerMask>(Flags))
+	{
+		if (HasAny<CASTCLASS_FArrayProperty>(Flags))
+		{
+			return BindArray(static_cast<FArrayProperty*>(Property), Ranges);
+		}
+		if (HasAny<CASTCLASS_FMapProperty>(Flags))
+		{
+			return BindMap(static_cast<FMapProperty*>(Property), Ranges);
+		}
+		return HasAny<CASTCLASS_FSetProperty>(Flags)
+			? BindSet(static_cast<FSetProperty*>(Property), Ranges)
+			: BindOptional(static_cast<FOptionalProperty*>(Property), Ranges);
+	}
+	else if (HasAny<StringMask>(Flags))
+	{
+		if (HasAny<CASTCLASS_FNameProperty>(Flags))
+		{
+			return BindName(static_cast<FNameProperty*>(Property));
+		}
+		if (HasAny<CASTCLASS_FStrProperty>(Flags))
+		{
+			return BindStr(static_cast<FStrProperty*>(Property), Ranges);
+		}
+		if (HasAny<CASTCLASS_FTextProperty>(Flags))
+		{
+			return BindText(static_cast<FTextProperty*>(Property), Ranges);
+		}
+		
+		return HasAny<CASTCLASS_FAnsiStrProperty>(Flags)
+			? BindStr(static_cast<FAnsiStrProperty*>(Property), Ranges)
+			: BindStr(static_cast<FUtf8StrProperty*>(Property), Ranges);
+	}
+	else if (HasAny<CASTCLASS_FFieldPathProperty >(Flags))
+	{
+		return BindFieldPath(static_cast<FFieldPathProperty*>(Property));
+	}
+#if WITH_VERSE_VM
+	else if (HasAny<CASTCLASS_FVValueProperty | CASTCLASS_FVRestValueProperty>(Flags))
+	{
+		return HasAny<CASTCLASS_FEnumProperty>(Flags)
+			? BindVValue(static_cast<FVValueProperty*>(Property), Ranges)
+			: BindVRestValue(static_cast<FVRestValueProperty*>(Property), Ranges); 
+	}
+#endif
+
+	checkf(false, TEXT("Unrecognized cast flags %llx in property %s"), Flags, *Property->GetNameCPP());
+	return FMemberBinding(Property->GetOffset_ForInternal());
+}
+
+static void BindStruct(FStructSchemaId Id, UStruct* Struct)
+{
+	TArray<FMemberBinding, TInlineAllocator<64>> Members;
+	FRangeBindingArray Ranges;
+	for (FProperty* Property = Struct->PropertyLink; Property; Property = Property->PropertyLinkNext)
+	{
+		Members.Add(BindProperty(Property, /* tmp */ Ranges));
+	}
+	GUE.Schemas.BindStruct(Id, Id, Members); 
+}
+
+////////////////////////////////////////////////////////////////////////////////////////////////
+
 static void BindInitialTypes()
 {
 	if (!FParse::Param(FCommandLine::Get(), TEXT("BindPlainProps")))
@@ -123,19 +296,28 @@ static void BindInitialTypes()
 		return;
 	}
 
-	UE_LOG(LogInit, Display, TEXT("Binding classes to PlainProps schemas"));
+	UE_LOG(LogInit, Display, TEXT("Binding types to PlainProps schemas"));
+
+	TArray<FStructSchemaId> Ids;
 	for (TObjectIterator<UField> It; It; ++It)
 	{
 		UObject* Object = *It; 
 		if (UStruct* Struct = Cast<UStruct>(Object))
 		{
-			DeclareStruct(Struct);
+			Ids.Add(DeclareStruct(Struct));
 		}
 		else if (UEnum* Enum = Cast<UEnum>(Object))
 		{
 			DeclareEnum(Enum);
 		}	
 	}
+
+	const FStructSchemaId* IdIt = &Ids[0];
+	for (TObjectIterator<UStruct> It; It; ++It)
+	{
+		BindStruct(*IdIt++, *It);
+	}
+	check(IdIt == &Ids[0] + Ids.Num());
 }
 
 static FDelegateHandle PreMainCallback = FCoreDelegates::OnCommandletPreMain.AddStatic(&BindInitialTypes);
