PlainProps initial implementation of member numerals and FProperty static array bindings

#rnx

diff --git a/Engine/Plugins/Experimental/PlainProps/Source/Private/PlainPropsBind.cpp b/Engine/Plugins/Experimental/PlainProps/Source/Private/PlainPropsBind.cpp
--- a/Engine/Plugins/Experimental/PlainProps/Source/Private/PlainPropsBind.cpp
+++ b/Engine/Plugins/Experimental/PlainProps/Source/Private/PlainPropsBind.cpp
@@ -357,7 +357,7 @@ static void TranslateScopeIds(TArrayView<FNestedScopeId> Out, FIdIndexerBase& In
 		check(Scope.Outer.IsFlat() || Scope.Outer.AsNested().Idx < OutIdx);
 		FScopeId Outer = Scope.Outer.IsFlat() ? FScopeId(Translate(Scope.Outer.AsFlat(), ToNames)) : FScopeId(Out[Scope.Outer.AsNested().Idx]);
 		FFlatScopeId Inner = Translate(Scope.Inner, ToNames);
-		Out[OutIdx++] = Indexer.NestScope(Outer, Inner).AsNested();
+		Out[OutIdx++] = Indexer.NestFlatScope(Outer, Inner).AsNested();
 	}
 }
 
diff --git a/Engine/Plugins/Experimental/PlainProps/Source/Private/PlainPropsDeclare.cpp b/Engine/Plugins/Experimental/PlainProps/Source/Private/PlainPropsDeclare.cpp
--- a/Engine/Plugins/Experimental/PlainProps/Source/Private/PlainPropsDeclare.cpp
+++ b/Engine/Plugins/Experimental/PlainProps/Source/Private/PlainPropsDeclare.cpp
@@ -97,33 +97,45 @@ void CopyItems(T* It, TConstArrayView<T> Items)
 	}
 }
 
-void FDeclarations::DeclareStruct(FStructSchemaId DeclId, FTypeId Type, TConstArrayView<FMemberId> MemberOrder, EMemberPresence Occupancy, FOptionalStructSchemaId Super)
+static void Declare(TArray<TUniquePtr<FStructDeclaration>>& Out, FStructSchemaId DeclId, FTypeId Type, TConstArrayView<FMemberId> MemberOrder, EMemberPresence Occupancy, FOptionalStructSchemaId Super, uint16 NumMembers)
 {
-	if (static_cast<int32>(DeclId.Idx) >= DeclaredStructs.Num())
+	if (static_cast<int32>(DeclId.Idx) >= Out.Num())
 	{
-		DeclaredStructs.SetNum(DeclId.Idx + 1);
+		Out.SetNum(DeclId.Idx + 1);
 	}
 
-	TUniquePtr<FStructDeclaration>& Ptr = DeclaredStructs[DeclId.Idx];
+	TUniquePtr<FStructDeclaration>& Ptr = Out[DeclId.Idx];
 	if (Ptr)
 	{
 		check(DeclId == Ptr->Id);
 		check(Type == Ptr->Type);
 		check(Super == Ptr->Super);
 		check(Occupancy == Ptr->Occupancy);
+		check(NumMembers == Ptr->NumMembers);
 		check(Algo::Compare(MemberOrder, Ptr->GetMemberOrder()));
 
 		Ptr->RefCount++;
 	}
 	else
 	{
-		FStructDeclaration Header{1, DeclId, Type, Super, Occupancy, IntCastChecked<uint16>(MemberOrder.Num())};
+		FStructDeclaration Header{1, DeclId, Type, Super, Occupancy, !MemberOrder.IsEmpty(), NumMembers};
 		void* Data = FMemory::Malloc(sizeof(FStructDeclaration) + MemberOrder.Num() * MemberOrder.GetTypeSize());
 		Ptr.Reset(new (Data) FStructDeclaration(Header));
 		CopyItems(Ptr->MemberOrder, MemberOrder);
 	}
 }
 
+
+void FDeclarations::DeclareStruct(FStructSchemaId DeclId, FTypeId Type, TConstArrayView<FMemberId> MemberOrder, EMemberPresence Occupancy, FOptionalStructSchemaId Super)
+{
+	Declare(/* out */ DeclaredStructs, DeclId, Type, MemberOrder, Occupancy, Super, IntCastChecked<uint16>(MemberOrder.Num()));
+}
+
+void FDeclarations::DeclareNumeralStruct(FStructSchemaId DeclId, FTypeId Type, uint16 NumMembers, EMemberPresence Occupancy)
+{
+	Declare(/* out */ DeclaredStructs, DeclId, Type, {}, Occupancy, NoId, NumMembers);
+}
+
 void FDeclarations::DeclareEnum(FEnumSchemaId Id, FTypeId Type, EEnumMode Mode, TConstArrayView<FEnumerator> Enumerators, EEnumAliases Policy)
 {
 	if (static_cast<int32>(Id.Idx) >= DeclaredEnums.Num())
diff --git a/Engine/Plugins/Experimental/PlainProps/Source/Private/PlainPropsIndex.cpp b/Engine/Plugins/Experimental/PlainProps/Source/Private/PlainPropsIndex.cpp
--- a/Engine/Plugins/Experimental/PlainProps/Source/Private/PlainPropsIndex.cpp
+++ b/Engine/Plugins/Experimental/PlainProps/Source/Private/PlainPropsIndex.cpp
@@ -122,7 +122,7 @@ FParametricTypeView FParametricTypeIndexer::Resolve(FParametricTypeId Id) const
 
 //////////////////////////////////////////////////////////////////////////
 
-FScopeId FIdIndexerBase::NestScope(FScopeId Outer, FFlatScopeId Inner)
+FScopeId FIdIndexerBase::NestFlatScope(FScopeId Outer, FFlatScopeId Inner)
 {
 	check(Outer);
 	return FScopeId(NestedScopes.Index(Outer, Inner));
@@ -133,7 +133,7 @@ FScopeId FIdIndexerBase::NestReversedScopes(TConstArrayView<FFlatScopeId> Inners
 	FScopeId Out = NoId;
 	for (FFlatScopeId Inner : ReverseIterate(Inners))
 	{
-		Out = Out ? NestScope(Out, Inner) : FScopeId(Inner);
+		Out = Out ? NestFlatScope(Out, Inner) : FScopeId(Inner);
 	}
 
 	return Out;
diff --git a/Engine/Plugins/Experimental/PlainProps/Source/Public/PlainPropsBind.h b/Engine/Plugins/Experimental/PlainProps/Source/Public/PlainPropsBind.h
--- a/Engine/Plugins/Experimental/PlainProps/Source/Public/PlainPropsBind.h
+++ b/Engine/Plugins/Experimental/PlainProps/Source/Public/PlainPropsBind.h
@@ -882,6 +882,8 @@ FStructSchemaId BindCustomStructOnce()
 	return Instance.BindId;
 }
 
+inline constexpr FMemberBindType DefaultStructBindType = FMemberBindType(FStructType{EMemberKind::Struct, /* IsDynamic */ 0, /* IsSuper */ 0});
+
 template<class Struct, class CustomBinding, class Runtime>
 FMemberBindType BindMemberStruct(FOptionalSchemaId& OutSchema)
 {
@@ -895,14 +897,14 @@ FMemberBindType BindMemberStruct(FOptionalSchemaId& OutSchema)
 		OutSchema = FOptionalSchemaId(BindCustomStructOnce<Struct, CustomBinding, Runtime>());
 	}
 
-	return FMemberBindType(FStructType{EMemberKind::Struct, /* IsDynamic */ 0, /* IsSuper */ 0});
+	return DefaultStructBindType;
 }
 
 template<typename Struct, class Ids>
 FMemberBindType BindInnermostType(FOptionalSchemaId& OutSchema)
 {
 	OutSchema = FOptionalSchemaId(GetStructBindId<Ids, Struct>());
-	return FMemberBindType(FStructType{EMemberKind::Struct, /* IsDynamic */ 0, /* IsSuper */ 0});
+	return DefaultStructBindType;
 }
 
 template<Arithmetic Type, class Ids>
diff --git a/Engine/Plugins/Experimental/PlainProps/Source/Public/PlainPropsDeclare.h b/Engine/Plugins/Experimental/PlainProps/Source/Public/PlainPropsDeclare.h
--- a/Engine/Plugins/Experimental/PlainProps/Source/Public/PlainPropsDeclare.h
+++ b/Engine/Plugins/Experimental/PlainProps/Source/Public/PlainPropsDeclare.h
@@ -32,7 +32,7 @@ struct FEnumDeclaration
 	}
 };
 
-enum class EMemberPresence { RequireAll, AllowSparse };
+enum class EMemberPresence : uint8 { RequireAll, AllowSparse };
 
 struct FStructDeclaration
 {
@@ -41,13 +41,17 @@ struct FStructDeclaration
 	FTypeId					Type;			// Could be removed, might allow declaration dedup among templated types
 	FOptionalStructSchemaId	Super;
 	EMemberPresence			Occupancy;
+	bool					bNamedMembers;
 	uint16					NumMembers;
-	FMemberId				MemberOrder[0];
+	FMemberId				MemberOrder[0]; // Only if bNamedMembers
 
 	TConstArrayView<FMemberId> GetMemberOrder() const
 	{
+		check(bNamedMembers);
 		return MakeArrayView(MemberOrder, NumMembers);
 	}
+
+	inline static constexpr uint16 MaxMembers = 0xFFFF;
 };
 
 // Enum values are stored as integers. Aliased enum values are illegal, including composite flags.
@@ -63,6 +67,7 @@ public:
 	PLAINPROPS_API void								DeclareEnum(FEnumSchemaId DeclId, FTypeId Type, EEnumMode Mode, TConstArrayView<FEnumerator> Enumerators, EEnumAliases Policy);
 	// Declare struct with ref count 1 or increment it and check that previous declaration matches
 	PLAINPROPS_API void								DeclareStruct(FStructSchemaId DeclId, FTypeId Type, TConstArrayView<FMemberId> MemberOrder, EMemberPresence Occupancy, FOptionalStructSchemaId Super = {});
+	PLAINPROPS_API void								DeclareNumeralStruct(FStructSchemaId DeclId, FTypeId Type, uint16 NumMembers, EMemberPresence Occupancy);
 	
 	void											DropEnum(FEnumSchemaId DeclId)			{ Check(DeclId); DeclaredEnums[DeclId.Idx].Reset(); }
 	void											DropStructRef(FStructSchemaId DeclId);
diff --git a/Engine/Plugins/Experimental/PlainProps/Source/Public/PlainPropsIndex.h b/Engine/Plugins/Experimental/PlainProps/Source/Public/PlainPropsIndex.h
--- a/Engine/Plugins/Experimental/PlainProps/Source/Public/PlainPropsIndex.h
+++ b/Engine/Plugins/Experimental/PlainProps/Source/Public/PlainPropsIndex.h
@@ -49,7 +49,7 @@ private:
 class FIdIndexerBase : public FDebugIds
 {
 public:
-	PLAINPROPS_API FScopeId				NestScope(FScopeId Outer, FFlatScopeId Inner);
+	PLAINPROPS_API FScopeId				NestFlatScope(FScopeId Outer, FFlatScopeId Inner);
 	PLAINPROPS_API FScopeId				NestReversedScopes(TConstArrayView<FFlatScopeId> Inners);
 	PLAINPROPS_API FParametricTypeId	MakeParametricTypeId(FOptionalConcreteTypenameId Name, TConstArrayView<FTypeId> Params);
 	PLAINPROPS_API FTypeId				MakeParametricType(FTypeId Type, TConstArrayView<FTypeId> Params);
@@ -90,7 +90,7 @@ public:
 	template<typename T> FNameId		MakeName(T&& Str)									{ return { IntCastChecked<uint32>(Names.Add(NameType(Str)).AsInteger()) }; }
 	template<typename T> FMemberId		NameMember(T&& Name)								{ return { MakeName(Name) }; }
 	template<typename T> FFlatScopeId	MakeScope(T&& Str)									{ return static_cast<FFlatScopeId>(MakeName(Str)); }
-	template<typename T> FScopeId		NestScope(FScopeId Outer, T&& Inner)				{ return FIdIndexerBase::NestScope(Outer, static_cast<FFlatScopeId>(MakeName(Inner))); }
+	template<typename T> FScopeId		NestScope(FScopeId Outer, T&& Inner)				{ return NestFlatScope(Outer, static_cast<FFlatScopeId>(MakeName(Inner))); }
 	template<typename T> FTypenameId	MakeTypename(T&& Name)								{ return FTypenameId(FConcreteTypenameId(MakeName(Name))); }
 	template<typename T> FTypeId		MakeType(T&& Scope, T&& Name)						{ return { FScopeId(MakeScope(Scope)), MakeTypename(Name) }; }
 
diff --git a/Engine/Plugins/Experimental/PlainPropsUObject/Source/Private/PlainPropsUObjectBindings.cpp b/Engine/Plugins/Experimental/PlainPropsUObject/Source/Private/PlainPropsUObjectBindings.cpp
--- a/Engine/Plugins/Experimental/PlainPropsUObject/Source/Private/PlainPropsUObjectBindings.cpp
+++ b/Engine/Plugins/Experimental/PlainPropsUObject/Source/Private/PlainPropsUObjectBindings.cpp
@@ -30,7 +30,7 @@ FGlobals GUE;
 
 ////////////////////////////////////////////////////////////////////////////////////////////////
 
-static FTypeId IndexType(UField* Field)
+static FTypeId IndexType(const UField* Field)
 {	
 	FTypenameId Name = GUE.Names.MakeTypename(Field->GetFName());
 	
@@ -44,22 +44,22 @@ static FTypeId IndexType(UField* Field)
 	return { GUE.Names.NestReversedScopes(ReversedOuters), Name };
 }
 
-static FOptionalStructSchemaId IndexSuper(UStruct* Struct)
+static FOptionalStructSchemaId IndexSuper(const UStruct* Struct)
 {
 	return Struct ? ToOptional(GUE.Names.IndexStruct(IndexType(Struct))) : NoId;
 }
 
-static EMemberPresence GetOccupancy(UStruct* Struct)
+static EMemberPresence GetOccupancy(const UStruct* Struct)
 {
 	if (Struct->HasAnyCastFlags(CASTCLASS_UScriptStruct))
 	{
-		EStructFlags Flags = static_cast<UScriptStruct*>(Struct)->StructFlags;
+		EStructFlags Flags = static_cast<const UScriptStruct*>(Struct)->StructFlags;
 		return (Flags & (STRUCT_Immutable | STRUCT_Atomic)) ? EMemberPresence::RequireAll : EMemberPresence::AllowSparse; 
 	}
 	return EMemberPresence::AllowSparse;
 }
 
-static FStructSchemaId DeclareStruct(UStruct* Struct)
+static FStructSchemaId DeclareStruct(const UStruct* Struct)
 {
 	FTypeId Type = IndexType(Struct);
 	FStructSchemaId Id = GUE.Names.IndexStruct(Type);
@@ -125,166 +125,338 @@ static FEnumSchemaId DeclareEnum(UEnum* Enum)
 
 ////////////////////////////////////////////////////////////////////////////////////////////////
 
-using FRangeBindingArray = TPagedArray<FRangeBinding, 1024>;
-
-inline FMemberBinding Todo(FProperty* Property)
+struct FStaticArrayItemBinding : IItemRangeBinding
 {
-	return FMemberBinding(Property->GetOffset_ForInternal());
-}
+	FStaticArrayItemBinding(uint32 InNum, uint32 InStride) : Num(InNum), Stride(InStride) {}
 
-inline FMemberBinding BindImpl(FEnumProperty* Property, FMemberBindType	InnermostType, FOptionalSchemaId InnermostSchema, TConstArrayView<FRangeBinding> RangeBindings)
-{
-	FMemberBinding Out(Property->GetOffset_ForInternal());
-	Out.InnermostType = InnermostType;
-	Out.InnermostSchema = InnermostSchema;
-	Out.RangeBindings = RangeBindings;
+	uint32 Num;
+	uint32 Stride;
 
-	if (Property->GetElementSize() == Property->GetSize())
+	virtual void ReadItems(FSaveRangeContext& Ctx) const override
 	{
-		return Out;
+		Ctx.Items.SetAll({Ctx.Request.Range, Num}, Stride);
 	}
-	
-	// If we hit this limit, generate one-level of nested structs
-	uint16 ArrayDim = IntCastChecked<uint16>(Property->GetSize() / Property->GetElementSize());
 
-	// todo
-	return Out;
-}
+	virtual void MakeItems(FLoadRangeContext& Ctx) const override {}
+};
+
 
+////////////////////////////////////////////////////////////////////////////////////////////////
 
-inline FMemberBinding BindNumeric(FNumericProperty* Property, uint64 CastFlags) { return Todo(Property); }
-inline FMemberBinding BindEnum(FEnumProperty* Property)
+static void BindStruct(FStructSchemaId Id, UStruct* Struct);
+
+class FPropertyBinder
 {
-	FSchemaId Id = GUE.Names.IndexEnum(IndexType(Property->GetEnum()));
-	FUnpackedLeafType Leaf = {ELeafType::Enum, WidthOf(Property->GetElementSize())};
-	return BindImpl(Property, FMemberBindType(Leaf), Id, {});
-}
-inline FMemberBinding BindBool(FBoolProperty* Property) { return Todo(Property); }
+public:
+	FPropertyBinder(FStructSchemaId Id) : Owner(Id) {}
 
-inline FMemberBinding BindReference(FObjectPropertyBase* Property, uint64 CastFlags) { return Todo(Property); }
-inline FMemberBinding BindInterface(FInterfaceProperty* Property) { return Todo(Property); }
+	void								BindMember(FProperty* Property) { Members.Emplace(BindProperty(Property)); }
+	TConstArrayView<FMemberBinding>		GetMembers() const { return Members; }
 
-inline FMemberBinding BindStruct(FStructProperty* Property, FRangeBindingArray& Ranges)	{ return Todo(Property); }
-inline FMemberBinding BindDelegate(FDelegateProperty* Property, FRangeBindingArray& Ranges) { return Todo(Property); }
-inline FMemberBinding BindMulticastDelegate(FMulticastDelegateProperty* Property, FRangeBindingArray& Ranges) { return Todo(Property); }
+private:
+	const FStructSchemaId							Owner;
+	TOptional<FScopeId>								OwnerScope;
+	TOptional<EMemberPresence>						OwnerOccupancy;
+	TPagedArray<FRangeBinding, 1024>				Ranges;
+	TArray<FMemberBinding, TInlineAllocator<64>>	Members;
 
-inline FMemberBinding BindArray(FArrayProperty* Property, FRangeBindingArray& Ranges) { return Todo(Property); }
-inline FMemberBinding BindMap(FMapProperty* Property, FRangeBindingArray& Ranges) { return Todo(Property); }
-inline FMemberBinding BindSet(FSetProperty* Property, FRangeBindingArray& Ranges) { return Todo(Property); }
-inline FMemberBinding BindOptional(FOptionalProperty* Property, FRangeBindingArray& Ranges) { return Todo(Property); }
+	TConstArrayView<FRangeBinding> AllocateRangeBindings(FRangeBinding Head, TConstArrayView<FRangeBinding> Tail)
+	{
+		if (Tail.Num() == 0)
+		{
+			return MakeArrayView(&Ranges.Add_GetRef(Head), 1);
+		}
+		
+		// Ensure contiguous range by skipping last items in page
+		static constexpr uint32 PageMax = decltype(Ranges)::MaxPerPage();
+		const int32 OutNum = 1 + Tail.Num();
+		check(OutNum <= PageMax);
+		if (Ranges.NumPages() != (Ranges.Num() + Tail.Num()) % PageMax)
+		{
+			int32 NumPad = Ranges.MaxPerPage() - Ranges.Num() % PageMax;
+			Ranges.Append(Tail.Slice(0, NumPad));
+			check(Ranges.Num() % PageMax == 0);
+		}
 
-template<typename StrProperty>
-inline FMemberBinding BindStr(StrProperty* Property, FRangeBindingArray& Ranges) { return Todo(Property); }
-inline FMemberBinding BindName(FNameProperty* Property) { return Todo(Property); }
-inline FMemberBinding BindText(FTextProperty* Property, FRangeBindingArray& Ranges) { return Todo(Property); }
+		const FRangeBinding* OutData = &Ranges.Add_GetRef(Head);
+		Ranges.Append(Tail);
+		return MakeArrayView(OutData, OutNum);
+	}
 
-inline FMemberBinding BindFieldPath(FFieldPathProperty* Property) { return Todo(Property); }
+	template<uint64 Mask>
+	static bool HasAny(uint64 Flags)
+	{
+		return (Mask & Flags) != 0;
+	}
 
-#if WITH_VERSE_VM
-inline FMemberBinding BindVValue(FVRestValueProperty* Property, FRangeBindingArray& Ranges);
-inline FMemberBinding BindVRestValue(FVRestValueProperty* Property, FRangeBindingArray& Ranges);
-#endif
+	static FMemberBinding Todo(FProperty* Property)
+	{
+		return FMemberBinding(Property->GetOffset_ForInternal());
+	}
 
-template<uint64 Mask>
-inline bool HasAny(uint64 Flags)
-{
-	return (Mask & Flags) != 0;
-}
+	FScopeId GetOwnerScope()
+	{
+		if (!OwnerScope)
+		{
+			FTypeId OwnerType = GUE.Names.Resolve(Owner);
+			OwnerScope = GUE.Names.NestFlatScope(OwnerType.Scope, {OwnerType.Name.AsConcrete().Id});
+		}
 
-static FMemberBinding BindProperty(FProperty* Property, FRangeBindingArray& Ranges)
-{
-	constexpr uint64 LeafMask = CASTCLASS_FNumericProperty | CASTCLASS_FEnumProperty| CASTCLASS_FBoolProperty;
-	constexpr uint64 ReferenceMask = CASTCLASS_FObjectPropertyBase | CASTCLASS_FInterfaceProperty;
-	constexpr uint64 ContainerMask = CASTCLASS_FArrayProperty | CASTCLASS_FSetProperty| CASTCLASS_FMapProperty  | CASTCLASS_FOptionalProperty;
-	constexpr uint64 StructuralMask = CASTCLASS_FStructProperty | CASTCLASS_FDelegateProperty | CASTCLASS_FMulticastDelegateProperty;
-	constexpr uint64 StringMask = CASTCLASS_FStrProperty | CASTCLASS_FUtf8StrProperty | CASTCLASS_FAnsiStrProperty | CASTCLASS_FNameProperty | CASTCLASS_FTextProperty;
-
-	uint64 Flags = Property->GetCastFlags();
-	if (HasAny<LeafMask>(Flags))
+		return OwnerScope.GetValue();
+	}
+
+	EMemberPresence GetOwnerOccupancy()
 	{
-		if (HasAny<CASTCLASS_FNumericProperty>(Flags))
+		if (!OwnerOccupancy)
 		{
-			return BindNumeric(static_cast<FNumericProperty*>(Property), Flags);
+			OwnerOccupancy = GUE.Types.Get(Owner).Occupancy;
 		}
-		return HasAny<CASTCLASS_FEnumProperty>(Flags)
-			? BindEnum(static_cast<FEnumProperty*>(Property))
-			: BindBool(static_cast<FBoolProperty*>(Property));
+
+		return OwnerOccupancy.GetValue();
+	}
+
+
+	inline static FBitfieldBoolBindType MakeBitfieldBool(uint8 BitIdx)
+	{
+		return {EMemberKind::Leaf, ELeafBindType::BitfieldBool, BitIdx};
 	}
-	else if (HasAny<ReferenceMask>(Flags))
+	
+	inline FMemberBinding BindBool(FBoolProperty* Property)
+	{
+		check(Property->GetByteOffset() == 0);
+		FMemberBinding Out(Property->GetOffset_ForInternal());
+		uint8 BitIdx = static_cast<uint8>(FMath::FloorLog2NonZero(Property->GetFieldMask()));
+		FLeafBindType Type = Property->IsNativeBool()	? FLeafBindType(ELeafBindType::Bool, ELeafWidth::B8)
+														: FLeafBindType(MakeBitfieldBool(BitIdx));
+		Out.InnermostType = FMemberBindType(Type);
+		return Out;
+	}
+
+	inline FMemberBinding BindEnum(FEnumProperty* Property)
 	{
-		return HasAny<CASTCLASS_FObjectPropertyBase>(Flags)
-			? BindReference(static_cast<FObjectPropertyBase*>(Property), Flags)
-			: BindInterface(static_cast<FInterfaceProperty*>(Property));
+		FMemberBinding Out(Property->GetOffset_ForInternal());
+		Out.InnermostSchema = GUE.Names.IndexEnum(IndexType(Property->GetEnum()));
+		FUnpackedLeafType Leaf = {ELeafType::Enum, WidthOf(Property->GetElementSize())};
+		Out.InnermostType = FMemberBindType(Leaf);
+		return Out;
 	}
-	else if (HasAny<StructuralMask>(Flags))
+
+	inline FLeafBindType BindByte(FByteProperty* Property, FOptionalSchemaId& OutEnumId)
 	{
-		if (HasAny<CASTCLASS_FStructProperty>(Flags))
+		if (const UEnum* Enum = Property->GetIntPropertyEnum())
 		{
-			return BindStruct(static_cast<FStructProperty*>(Property), Ranges);
+			OutEnumId = GUE.Names.IndexEnum(IndexType(Enum));
+			return FLeafBindType(ELeafBindType::Enum, ELeafWidth::B8);
 		}
-		return HasAny<CASTCLASS_FDelegateProperty>(Flags)
-			? BindDelegate(static_cast<FDelegateProperty*>(Property), Ranges)
-			: BindMulticastDelegate(static_cast<FMulticastDelegateProperty*>(Property), Ranges);
+		return FLeafBindType(ELeafBindType::IntU, ELeafWidth::B8);
 	}
-	else if (HasAny<ContainerMask>(Flags))
+
+	inline FMemberBinding BindNumeric(FNumericProperty* Property, uint64 Flags)
 	{
-		if (HasAny<CASTCLASS_FArrayProperty>(Flags))
+		FMemberBinding Out(Property->GetOffset_ForInternal());
+		bool bFloat = HasAny<CASTCLASS_FFloatProperty | CASTCLASS_FDoubleProperty>(Flags);
+		bool bIntS = HasAny<CASTCLASS_FInt8Property | CASTCLASS_FInt16Property | CASTCLASS_FIntProperty | CASTCLASS_FInt64Property>(Flags);
+		if (HasAny<CASTCLASS_FByteProperty>(Flags))
 		{
-			return BindArray(static_cast<FArrayProperty*>(Property), Ranges);
+			FLeafBindType Leaf = BindByte(static_cast<FByteProperty*>(Property), /* out enum */ Out.InnermostSchema);
+			Out.InnermostType = FMemberBindType(Leaf);
 		}
-		if (HasAny<CASTCLASS_FMapProperty>(Flags))
+		else
 		{
-			return BindMap(static_cast<FMapProperty*>(Property), Ranges);
+			ELeafType Type = bFloat ? ELeafType::Float : (bIntS ? ELeafType::IntS : ELeafType::IntU);
+			FUnpackedLeafType Leaf = {Type, WidthOf(Property->GetElementSize())};
+			Out.InnermostType = FMemberBindType(Leaf);	
 		}
-		return HasAny<CASTCLASS_FSetProperty>(Flags)
-			? BindSet(static_cast<FSetProperty*>(Property), Ranges)
-			: BindOptional(static_cast<FOptionalProperty*>(Property), Ranges);
+		
+		return Out;
 	}
-	else if (HasAny<StringMask>(Flags))
+
+	inline FMemberBinding BindReference(FObjectPropertyBase* Property, uint64 CastFlags) { return Todo(Property); }
+	inline FMemberBinding BindInterface(FInterfaceProperty* Property) { return Todo(Property); }
+
+	inline FMemberBinding BindStruct(FStructProperty* Property)	const { return Todo(Property); }
+	inline FMemberBinding BindDelegate(FDelegateProperty* Property) { return Todo(Property); }
+	inline FMemberBinding BindMulticastDelegate(FMulticastDelegateProperty* Property) { return Todo(Property); }
+
+	inline FMemberBinding BindArray(FArrayProperty* Property) { return Todo(Property); }
+	inline FMemberBinding BindMap(FMapProperty* Property) { return Todo(Property); }
+	inline FMemberBinding BindSet(FSetProperty* Property) { return Todo(Property); }
+	inline FMemberBinding BindOptional(FOptionalProperty* Property) { return Todo(Property); }
+
+	template<typename StrProperty>
+	inline FMemberBinding BindStr(StrProperty* Property) { return Todo(Property); }
+	inline FMemberBinding BindName(FNameProperty* Property) { return Todo(Property); }
+	inline FMemberBinding BindText(FTextProperty* Property) { return Todo(Property); }
+
+	inline FMemberBinding BindFieldPath(FFieldPathProperty* Property) { return Todo(Property); }
+
+	#if WITH_VERSE_VM
+	inline FMemberBinding BindVValue(FVRestValueProperty* Property) { return Todo(Property); }
+	inline FMemberBinding BindVRestValue(FVRestValueProperty* Property) { return Todo(Property); }
+	#endif
+
+	FMemberBinding BindSingleProperty(FProperty* Property)
 	{
-		if (HasAny<CASTCLASS_FNameProperty>(Flags))
+		constexpr uint64 LeafMask = CASTCLASS_FNumericProperty | CASTCLASS_FEnumProperty| CASTCLASS_FBoolProperty;
+		constexpr uint64 ReferenceMask = CASTCLASS_FObjectPropertyBase | CASTCLASS_FInterfaceProperty;
+		constexpr uint64 ContainerMask = CASTCLASS_FArrayProperty | CASTCLASS_FSetProperty| CASTCLASS_FMapProperty  | CASTCLASS_FOptionalProperty;
+		constexpr uint64 StructuralMask = CASTCLASS_FStructProperty | CASTCLASS_FDelegateProperty | CASTCLASS_FMulticastDelegateProperty;
+		constexpr uint64 StringMask = CASTCLASS_FStrProperty | CASTCLASS_FUtf8StrProperty | CASTCLASS_FAnsiStrProperty | CASTCLASS_FNameProperty | CASTCLASS_FTextProperty;
+
+		uint64 Flags = Property->GetCastFlags();
+		if (HasAny<LeafMask>(Flags))
 		{
-			return BindName(static_cast<FNameProperty*>(Property));
+			if (HasAny<CASTCLASS_FNumericProperty>(Flags))
+			{
+				return BindNumeric(static_cast<FNumericProperty*>(Property), Flags);
+			}
+			return HasAny<CASTCLASS_FEnumProperty>(Flags)
+				? BindEnum(static_cast<FEnumProperty*>(Property))
+				: BindBool(static_cast<FBoolProperty*>(Property));
 		}
-		if (HasAny<CASTCLASS_FStrProperty>(Flags))
+		else if (HasAny<ReferenceMask>(Flags))
 		{
-			return BindStr(static_cast<FStrProperty*>(Property), Ranges);
+			return HasAny<CASTCLASS_FObjectPropertyBase>(Flags)
+				? BindReference(static_cast<FObjectPropertyBase*>(Property), Flags)
+				: BindInterface(static_cast<FInterfaceProperty*>(Property));
 		}
-		if (HasAny<CASTCLASS_FTextProperty>(Flags))
+		else if (HasAny<StructuralMask>(Flags))
 		{
-			return BindText(static_cast<FTextProperty*>(Property), Ranges);
+			if (HasAny<CASTCLASS_FStructProperty>(Flags))
+			{
+				return BindStruct(static_cast<FStructProperty*>(Property));
+			}
+			return HasAny<CASTCLASS_FDelegateProperty>(Flags)
+				? BindDelegate(static_cast<FDelegateProperty*>(Property))
+				: BindMulticastDelegate(static_cast<FMulticastDelegateProperty*>(Property));
 		}
+		else if (HasAny<ContainerMask>(Flags))
+		{
+			if (HasAny<CASTCLASS_FArrayProperty>(Flags))
+			{
+				return BindArray(static_cast<FArrayProperty*>(Property));
+			}
+			if (HasAny<CASTCLASS_FMapProperty>(Flags))
+			{
+				return BindMap(static_cast<FMapProperty*>(Property));
+			}
+			return HasAny<CASTCLASS_FSetProperty>(Flags)
+				? BindSet(static_cast<FSetProperty*>(Property))
+				: BindOptional(static_cast<FOptionalProperty*>(Property));
+		}
+		else if (HasAny<StringMask>(Flags))
+		{
+			if (HasAny<CASTCLASS_FNameProperty>(Flags))
+			{
+				return BindName(static_cast<FNameProperty*>(Property));
+			}
+			if (HasAny<CASTCLASS_FStrProperty>(Flags))
+			{
+				return BindStr(static_cast<FStrProperty*>(Property));
+			}
+			if (HasAny<CASTCLASS_FTextProperty>(Flags))
+			{
+				return BindText(static_cast<FTextProperty*>(Property));
+			}
 		
-		return HasAny<CASTCLASS_FAnsiStrProperty>(Flags)
-			? BindStr(static_cast<FAnsiStrProperty*>(Property), Ranges)
-			: BindStr(static_cast<FUtf8StrProperty*>(Property), Ranges);
+			return HasAny<CASTCLASS_FAnsiStrProperty>(Flags)
+				? BindStr(static_cast<FAnsiStrProperty*>(Property))
+				: BindStr(static_cast<FUtf8StrProperty*>(Property));
+		}
+		else if (HasAny<CASTCLASS_FFieldPathProperty >(Flags))
+		{
+			return BindFieldPath(static_cast<FFieldPathProperty*>(Property));
+		}
+	#if WITH_VERSE_VM
+		else if (HasAny<CASTCLASS_FVValueProperty | CASTCLASS_FVRestValueProperty>(Flags))
+		{
+			return HasAny<CASTCLASS_FEnumProperty>(Flags)
+				? BindVValue(static_cast<FVValueProperty*>(Property))
+				: BindVRestValue(static_cast<FVRestValueProperty*>(Property)); 
+		}
+	#endif
+
+		checkf(false, TEXT("Unrecognized cast flags %llx in property %s"), Flags, *Property->GetNameCPP());
+		return FMemberBinding(Property->GetOffset_ForInternal());
 	}
-	else if (HasAny<CASTCLASS_FFieldPathProperty >(Flags))
+
+
+	FTypeId MakeStaticArrayTypename(FName PropertyName)
 	{
-		return BindFieldPath(static_cast<FFieldPathProperty*>(Property));
+		return { GetOwnerScope(), GUE.Names.MakeTypename(PropertyName) };
 	}
-#if WITH_VERSE_VM
-	else if (HasAny<CASTCLASS_FVValueProperty | CASTCLASS_FVRestValueProperty>(Flags))
+
+	FMemberBinding BindProperty(FProperty* Property)
 	{
-		return HasAny<CASTCLASS_FEnumProperty>(Flags)
-			? BindVValue(static_cast<FVValueProperty*>(Property), Ranges)
-			: BindVRestValue(static_cast<FVRestValueProperty*>(Property), Ranges); 
-	}
-#endif
+		FMemberBinding Out = BindSingleProperty(Property);
+		if (Property->GetSize() == Property->GetElementSize())
+		{
+			return Out;
+		}
 
-	checkf(false, TEXT("Unrecognized cast flags %llx in property %s"), Flags, *Property->GetNameCPP());
-	return FMemberBinding(Property->GetOffset_ForInternal());
-}
+		// Bind static array
+		EMemberPresence Occupancy = GetOwnerOccupancy();
+		uint32 TotalSize = static_cast<uint32>(Property->GetSize());
+		uint32 ElementSize = static_cast<uint32>(Property->GetElementSize());
+		uint32 ArrayDim = TotalSize / ElementSize;
+		check(ArrayDim * ElementSize == TotalSize);
+		if (Occupancy == EMemberPresence::RequireAll || ArrayDim > FStructDeclaration::MaxMembers)
+		{
+			// Create range binding that isn't delta-serializable
+			//
+			// Could generate nested numeral structs instead. Unsure if automatic
+			// per-element delta serialization for massive arrays is desirable.
+			//
+			// To delta-serialize massive arrays, custom-bind the owning struct
+			// and implement delta serialization manually
+
+			// Todo: Ownership / memory leak
+			const FStaticArrayItemBinding& ItemBinding = *new FStaticArrayItemBinding(ArrayDim, ElementSize);
+			ERangeSizeType SizeType = ArrayDim < 256 ? ERangeSizeType::U8 : ((ArrayDim < 65536) ? ERangeSizeType::U16 : ERangeSizeType::U32);
+			Out.RangeBindings =	AllocateRangeBindings(FRangeBinding(ItemBinding, SizeType), Out.RangeBindings);
+		}
+		else
+		{
+			// Create struct binding to allow delta serialization
+			FTypeId StaticArrayType = MakeStaticArrayTypename(Property->GetFName());
+			FStructSchemaId StaticArrayId = GUE.Names.IndexStruct(StaticArrayType);
+
+			// Todo: Ownership
+			GUE.Types.DeclareNumeralStruct(StaticArrayId, StaticArrayType, ArrayDim, Occupancy);
+
+			TArray<FMemberBinding, TInlineAllocator<64>> Elements;
+			Elements.Init(Out, ArrayDim);
+			uint64 Offset = 0;
+			for (FMemberBinding& Element : Elements)
+			{
+				Element.Offset = Offset;
+				Offset += ElementSize;
+			}
+			GUE.Schemas.BindStruct(StaticArrayId, StaticArrayId, Elements);
+			
+			Out.InnermostType = DefaultStructBindType;
+			Out.InnermostSchema = static_cast<FSchemaId>(StaticArrayId);
+			Out.RangeBindings = {};
+		}
+		
+	
+		return Out;	
+	}
+};
 
-static void BindStruct(FStructSchemaId Id, UStruct* Struct)
+void BindStruct(FStructSchemaId Id, UStruct* Struct)
 {
-	TArray<FMemberBinding, TInlineAllocator<64>> Members;
-	FRangeBindingArray Ranges;
+	if (GUE.Customs.FindStruct(Id))
+	{
+		return;
+	}
+
+	FPropertyBinder Binder(Id);
 	for (FProperty* Property = Struct->PropertyLink; Property; Property = Property->PropertyLinkNext)
 	{
-		Members.Add(BindProperty(Property, /* tmp */ Ranges));
+		Binder.BindMember(Property);
 	}
-	GUE.Schemas.BindStruct(Id, Id, Members); 
+	GUE.Schemas.BindStruct(Id, Id, Binder.GetMembers()); 
 }
 
 ////////////////////////////////////////////////////////////////////////////////////////////////
