PlainProps add FProperty bindings for FName, FText, F[Weak|Soft|Lazy]ObjectPtr, FFieldPath and FScriptInterface

Note: Haven't added default struct bindings for these types, only FProperty -> struct bindings

#rnx

diff --git a/Engine/Plugins/Experimental/PlainPropsUObject/Source/Private/PlainPropsUObjectBindings.cpp b/Engine/Plugins/Experimental/PlainPropsUObject/Source/Private/PlainPropsUObjectBindings.cpp
--- a/Engine/Plugins/Experimental/PlainPropsUObject/Source/Private/PlainPropsUObjectBindings.cpp
+++ b/Engine/Plugins/Experimental/PlainPropsUObject/Source/Private/PlainPropsUObjectBindings.cpp
@@ -19,11 +19,28 @@
 namespace PlainProps::UE
 {
 
+FCommonUObjectIds::FCommonUObjectIds(TIdIndexer<FName>& Names)
+: CoreUObject(Names.MakeScope("/Script/CoreUObject"))
+, Name(				Names.IndexStruct({CoreUObject, Names.MakeTypename("Name")}))
+, Text(				Names.IndexStruct({CoreUObject, Names.MakeTypename("Text")}))
+, FieldPath(		Names.IndexStruct({CoreUObject, Names.MakeTypename("FieldPath")}))
+, SoftObjectPath(	Names.IndexStruct({CoreUObject, Names.MakeTypename("SoftObjectPath")}))
+, ObjectPtr(		Names.IndexStruct({CoreUObject, Names.MakeTypename("ObjectPtr")}))
+, WeakObjectPtr(	Names.IndexStruct({CoreUObject, Names.MakeTypename("WeakObjectPtr")}))
+, LazyObjectPtr(	Names.IndexStruct({CoreUObject, Names.MakeTypename("LazyObjectPtr")}))
+, SoftObjectPtr(	Names.IndexStruct({CoreUObject, Names.MakeTypename("SoftObjectPtr")}))
+, ScriptInterface(	Names.IndexStruct({CoreUObject, Names.MakeTypename("ScriptInterface")}))
+{}
+
+
+////////////////////////////////////////////////////////////////////////////////////////////////
+
 FGlobals::FGlobals() 
 : Types(Names)
 , Schemas(Names)
 , Customs(Names)
 , DeltaCustoms(Names, &Customs)
+, CommonIds(Names)
 {}
 
 FGlobals GUE;
@@ -141,6 +158,34 @@ struct FStaticArrayItemBinding : IItemRangeBinding
 };
 
 
+////////////////////////////////////////////////////////////////////////////////////////////////
+
+static constexpr uint64 BindAsStructMask = 
+	CASTCLASS_FNameProperty | CASTCLASS_FTextProperty | CASTCLASS_FInterfaceProperty | CASTCLASS_FFieldPathProperty |
+	CASTCLASS_FObjectProperty | CASTCLASS_FWeakObjectProperty | CASTCLASS_FSoftObjectProperty | CASTCLASS_FLazyObjectProperty;
+
+static FStructSchemaId CastFlagsToStructBindId(uint64 MaskedCastFlags)
+{
+	switch (MaskedCastFlags)
+	{
+		case CASTCLASS_FNameProperty:			return GUE.CommonIds.Name;
+		case CASTCLASS_FTextProperty:			return GUE.CommonIds.Text;
+		case CASTCLASS_FInterfaceProperty:		return GUE.CommonIds.ScriptInterface;
+		case CASTCLASS_FFieldPathProperty:		return GUE.CommonIds.FieldPath;
+		case CASTCLASS_FObjectProperty:			return GUE.CommonIds.ObjectPtr;
+		case CASTCLASS_FWeakObjectProperty:		return GUE.CommonIds.WeakObjectPtr;
+		case CASTCLASS_FSoftObjectProperty:		return GUE.CommonIds.SoftObjectPtr;
+		case CASTCLASS_FLazyObjectProperty:		return GUE.CommonIds.LazyObjectPtr;
+		default:								break; // error
+	}
+
+	check(MaskedCastFlags); // @pre violated
+	check((MaskedCastFlags & BindAsStructMask) == MaskedCastFlags); // @pre violated
+	checkf(FMath::CountBits(MaskedCastFlags) == 1, TEXT("Masked CASTCLASS flags %llx match more than one property type"), MaskedCastFlags);
+	check(false); // Mismatch between this function and BindAsStructMask
+	return {};
+}
+
 ////////////////////////////////////////////////////////////////////////////////////////////////
 
 static void BindStruct(FStructSchemaId Id, UStruct* Struct);
@@ -271,9 +316,6 @@ private:
 		return Out;
 	}
 
-	inline FMemberBinding BindReference(FObjectPropertyBase* Property, uint64 CastFlags) { return Todo(Property); }
-	inline FMemberBinding BindInterface(FInterfaceProperty* Property) { return Todo(Property); }
-
 	inline FMemberBinding BindStruct(FStructProperty* Property)	const { return Todo(Property); }
 	inline FMemberBinding BindDelegate(FDelegateProperty* Property) { return Todo(Property); }
 	inline FMemberBinding BindMulticastDelegate(FMulticastDelegateProperty* Property) { return Todo(Property); }
@@ -285,26 +327,33 @@ private:
 
 	template<typename StrProperty>
 	inline FMemberBinding BindStr(StrProperty* Property) { return Todo(Property); }
-	inline FMemberBinding BindName(FNameProperty* Property) { return Todo(Property); }
-	inline FMemberBinding BindText(FTextProperty* Property) { return Todo(Property); }
-
-	inline FMemberBinding BindFieldPath(FFieldPathProperty* Property) { return Todo(Property); }
 
 	#if WITH_VERSE_VM
 	inline FMemberBinding BindVValue(FVRestValueProperty* Property) { return Todo(Property); }
 	inline FMemberBinding BindVRestValue(FVRestValueProperty* Property) { return Todo(Property); }
 	#endif
 
+	inline FMemberBinding BindAsStruct(FProperty* Property, FStructSchemaId Id)
+	{
+		FMemberBinding Out(Property->GetOffset_ForInternal());
+		Out.InnermostSchema = Id;
+		Out.InnermostType = DefaultStructBindType;
+		return Out;
+	}
+
 	FMemberBinding BindSingleProperty(FProperty* Property)
 	{
 		constexpr uint64 LeafMask = CASTCLASS_FNumericProperty | CASTCLASS_FEnumProperty| CASTCLASS_FBoolProperty;
-		constexpr uint64 ReferenceMask = CASTCLASS_FObjectPropertyBase | CASTCLASS_FInterfaceProperty;
 		constexpr uint64 ContainerMask = CASTCLASS_FArrayProperty | CASTCLASS_FSetProperty| CASTCLASS_FMapProperty  | CASTCLASS_FOptionalProperty;
 		constexpr uint64 StructuralMask = CASTCLASS_FStructProperty | CASTCLASS_FDelegateProperty | CASTCLASS_FMulticastDelegateProperty;
-		constexpr uint64 StringMask = CASTCLASS_FStrProperty | CASTCLASS_FUtf8StrProperty | CASTCLASS_FAnsiStrProperty | CASTCLASS_FNameProperty | CASTCLASS_FTextProperty;
+		constexpr uint64 StringMask = CASTCLASS_FStrProperty | CASTCLASS_FUtf8StrProperty | CASTCLASS_FAnsiStrProperty;
 
 		uint64 Flags = Property->GetCastFlags();
-		if (HasAny<LeafMask>(Flags))
+		if (HasAny<BindAsStructMask>(Flags))
+		{
+			return BindAsStruct(Property, CastFlagsToStructBindId(Flags & BindAsStructMask));
+		}
+		else if (HasAny<LeafMask>(Flags))
 		{
 			if (HasAny<CASTCLASS_FNumericProperty>(Flags))
 			{
@@ -314,12 +363,6 @@ private:
 				? BindEnum(static_cast<FEnumProperty*>(Property))
 				: BindBool(static_cast<FBoolProperty*>(Property));
 		}
-		else if (HasAny<ReferenceMask>(Flags))
-		{
-			return HasAny<CASTCLASS_FObjectPropertyBase>(Flags)
-				? BindReference(static_cast<FObjectPropertyBase*>(Property), Flags)
-				: BindInterface(static_cast<FInterfaceProperty*>(Property));
-		}
 		else if (HasAny<StructuralMask>(Flags))
 		{
 			if (HasAny<CASTCLASS_FStructProperty>(Flags))
@@ -346,27 +389,15 @@ private:
 		}
 		else if (HasAny<StringMask>(Flags))
 		{
-			if (HasAny<CASTCLASS_FNameProperty>(Flags))
-			{
-				return BindName(static_cast<FNameProperty*>(Property));
-			}
 			if (HasAny<CASTCLASS_FStrProperty>(Flags))
 			{
 				return BindStr(static_cast<FStrProperty*>(Property));
 			}
-			if (HasAny<CASTCLASS_FTextProperty>(Flags))
-			{
-				return BindText(static_cast<FTextProperty*>(Property));
-			}
 		
 			return HasAny<CASTCLASS_FAnsiStrProperty>(Flags)
 				? BindStr(static_cast<FAnsiStrProperty*>(Property))
 				: BindStr(static_cast<FUtf8StrProperty*>(Property));
 		}
-		else if (HasAny<CASTCLASS_FFieldPathProperty >(Flags))
-		{
-			return BindFieldPath(static_cast<FFieldPathProperty*>(Property));
-		}
 	#if WITH_VERSE_VM
 		else if (HasAny<CASTCLASS_FVValueProperty | CASTCLASS_FVRestValueProperty>(Flags))
 		{
diff --git a/Engine/Plugins/Experimental/PlainPropsUObject/Source/Public/PlainPropsUObjectRuntime.h b/Engine/Plugins/Experimental/PlainPropsUObject/Source/Public/PlainPropsUObjectRuntime.h
--- a/Engine/Plugins/Experimental/PlainPropsUObject/Source/Public/PlainPropsUObjectRuntime.h
+++ b/Engine/Plugins/Experimental/PlainPropsUObject/Source/Public/PlainPropsUObjectRuntime.h
@@ -16,6 +16,22 @@ template<> inline void AppendString(FString& Out, const FName& Name) { Name.Appe
 namespace PlainProps::UE
 {
 
+struct FCommonUObjectIds
+{
+	explicit FCommonUObjectIds(TIdIndexer<FName>& Names);
+
+	FScopeId			CoreUObject;
+	FStructSchemaId		Name;
+	FStructSchemaId		Text;
+	FStructSchemaId		FieldPath;
+	FStructSchemaId		SoftObjectPath;
+	FStructSchemaId		ObjectPtr;
+	FStructSchemaId		WeakObjectPtr;
+	FStructSchemaId		LazyObjectPtr;
+	FStructSchemaId		SoftObjectPtr;
+	FStructSchemaId		ScriptInterface;
+};
+
 struct FGlobals
 {
 	FGlobals();
@@ -25,6 +41,7 @@ struct FGlobals
 	FSchemaBindings		Schemas;
 	FCustomBindings		Customs;
 	FCustomBindings		DeltaCustoms;
+	FCommonUObjectIds	CommonIds;
 };
 PLAINPROPSUOBJECT_API extern FGlobals GUE;
 
