PlainProps add FStructProperty and F[Multicast]Delegate bindings

#rnx

diff --git a/Engine/Plugins/Experimental/PlainPropsUObject/Source/Private/PlainPropsUObjectBindings.cpp b/Engine/Plugins/Experimental/PlainPropsUObject/Source/Private/PlainPropsUObjectBindings.cpp
--- a/Engine/Plugins/Experimental/PlainPropsUObject/Source/Private/PlainPropsUObjectBindings.cpp
+++ b/Engine/Plugins/Experimental/PlainPropsUObject/Source/Private/PlainPropsUObjectBindings.cpp
@@ -160,11 +160,11 @@ struct FStaticArrayItemBinding : IItemRangeBinding
 
 ////////////////////////////////////////////////////////////////////////////////////////////////
 
-static constexpr uint64 BindAsStructMask = 
+static constexpr uint64 CommonStructMask = 
 	CASTCLASS_FNameProperty | CASTCLASS_FTextProperty | CASTCLASS_FInterfaceProperty | CASTCLASS_FFieldPathProperty |
 	CASTCLASS_FObjectProperty | CASTCLASS_FWeakObjectProperty | CASTCLASS_FSoftObjectProperty | CASTCLASS_FLazyObjectProperty;
 
-static FStructSchemaId CastFlagsToStructBindId(uint64 MaskedCastFlags)
+static FStructSchemaId FlagsToCommonBindId(uint64 MaskedCastFlags)
 {
 	switch (MaskedCastFlags)
 	{
@@ -180,8 +180,8 @@ static FStructSchemaId CastFlagsToStructBindId(uint64 MaskedCastFlags)
 	}
 
 	check(MaskedCastFlags); // @pre violated
-	check((MaskedCastFlags & BindAsStructMask) == MaskedCastFlags); // @pre violated
-	checkf(FMath::CountBits(MaskedCastFlags) == 1, TEXT("Masked CASTCLASS flags %llx match more than one property type"), MaskedCastFlags);
+	check((MaskedCastFlags & CommonStructMask) == MaskedCastFlags); // @pre violated
+	checkf(FMath::CountBits(MaskedCastFlags) == 1, TEXT("Masked CASTCLASS flags %llx match more than one common property type"), MaskedCastFlags);
 	check(false); // Mismatch between this function and BindAsStructMask
 	return {};
 }
@@ -316,10 +316,6 @@ private:
 		return Out;
 	}
 
-	inline FMemberBinding BindStruct(FStructProperty* Property)	const { return Todo(Property); }
-	inline FMemberBinding BindDelegate(FDelegateProperty* Property) { return Todo(Property); }
-	inline FMemberBinding BindMulticastDelegate(FMulticastDelegateProperty* Property) { return Todo(Property); }
-
 	inline FMemberBinding BindArray(FArrayProperty* Property) { return Todo(Property); }
 	inline FMemberBinding BindMap(FMapProperty* Property) { return Todo(Property); }
 	inline FMemberBinding BindSet(FSetProperty* Property) { return Todo(Property); }
@@ -341,6 +337,13 @@ private:
 		return Out;
 	}
 
+	inline FMemberBinding BindAsStruct(FProperty* Property, UStruct* Struct)
+	{
+		check(Struct);
+		FTypeId Type = IndexType(Struct);
+		return BindAsStruct(Property, GUE.Names.IndexStruct(Type));
+	}
+
 	FMemberBinding BindSingleProperty(FProperty* Property)
 	{
 		constexpr uint64 LeafMask = CASTCLASS_FNumericProperty | CASTCLASS_FEnumProperty| CASTCLASS_FBoolProperty;
@@ -349,11 +352,7 @@ private:
 		constexpr uint64 StringMask = CASTCLASS_FStrProperty | CASTCLASS_FUtf8StrProperty | CASTCLASS_FAnsiStrProperty;
 
 		uint64 Flags = Property->GetCastFlags();
-		if (HasAny<BindAsStructMask>(Flags))
-		{
-			return BindAsStruct(Property, CastFlagsToStructBindId(Flags & BindAsStructMask));
-		}
-		else if (HasAny<LeafMask>(Flags))
+		if (HasAny<LeafMask>(Flags))
 		{
 			if (HasAny<CASTCLASS_FNumericProperty>(Flags))
 			{
@@ -363,15 +362,19 @@ private:
 				? BindEnum(static_cast<FEnumProperty*>(Property))
 				: BindBool(static_cast<FBoolProperty*>(Property));
 		}
+		else if (HasAny<CommonStructMask>(Flags))
+		{
+			return BindAsStruct(Property, FlagsToCommonBindId(Flags & CommonStructMask));
+		}
 		else if (HasAny<StructuralMask>(Flags))
 		{
 			if (HasAny<CASTCLASS_FStructProperty>(Flags))
 			{
-				return BindStruct(static_cast<FStructProperty*>(Property));
+				return BindAsStruct(Property, static_cast<FStructProperty*>(Property)->Struct);
 			}
-			return HasAny<CASTCLASS_FDelegateProperty>(Flags)
-				? BindDelegate(static_cast<FDelegateProperty*>(Property))
-				: BindMulticastDelegate(static_cast<FMulticastDelegateProperty*>(Property));
+			return BindAsStruct(Property, HasAny<CASTCLASS_FDelegateProperty>(Flags) 
+								? static_cast<FDelegateProperty*>(Property)->SignatureFunction
+								: static_cast<FMulticastDelegateProperty*>(Property)->SignatureFunction);
 		}
 		else if (HasAny<ContainerMask>(Flags))
 		{
