PlainProps add FArrayProperty bindings and FLeafRangeLoadView::Copy

#rnx

diff --git a/Engine/Plugins/Experimental/PlainProps/Source/Private/PlainPropsBind.cpp b/Engine/Plugins/Experimental/PlainProps/Source/Private/PlainPropsBind.cpp
--- a/Engine/Plugins/Experimental/PlainProps/Source/Private/PlainPropsBind.cpp
+++ b/Engine/Plugins/Experimental/PlainProps/Source/Private/PlainPropsBind.cpp
@@ -138,10 +138,10 @@ FRangeBinding::FRangeBinding(const ILeafRangeBinding& Binding, ERangeSizeType Si
 
 ////////////////////////////////////////////////////////////////////////////////////////////////
 
-void* FLeafRangeAllocator::Allocate(FUnpackedLeafType Leaf, uint64 Num)
+void* FLeafRangeAllocator::Allocate(uint64 Num, SIZE_T LeafSize)
 {
 	check(!Range);
-	Range = FBuiltRange::Create(Scratch, Num, SizeOf(Leaf.Width));
+	Range = FBuiltRange::Create(Scratch, Num, LeafSize);
 	return Range->Data;
 }
 
diff --git a/Engine/Plugins/Experimental/PlainProps/Source/Private/PlainPropsSave.cpp b/Engine/Plugins/Experimental/PlainProps/Source/Private/PlainPropsSave.cpp
--- a/Engine/Plugins/Experimental/PlainProps/Source/Private/PlainPropsSave.cpp
+++ b/Engine/Plugins/Experimental/PlainProps/Source/Private/PlainPropsSave.cpp
@@ -465,7 +465,7 @@ static bool DiffItem(const uint8* A, const uint8* B, const FSaveContext& Ctx, FR
 
 	if (Binding.IsLeafBinding())
 	{
-		return !!Binding.AsLeafBinding().DiffLeaves(A, B);
+		return Binding.AsLeafBinding().DiffLeaves(A, B);
 	}
 
 	const IItemRangeBinding& ItemBinding = Binding.AsItemBinding();
diff --git a/Engine/Plugins/Experimental/PlainProps/Source/Public/PlainPropsBind.h b/Engine/Plugins/Experimental/PlainProps/Source/Public/PlainPropsBind.h
--- a/Engine/Plugins/Experimental/PlainProps/Source/Public/PlainPropsBind.h
+++ b/Engine/Plugins/Experimental/PlainProps/Source/Public/PlainPropsBind.h
@@ -466,6 +466,7 @@ struct FSaveRangeContext
 
 struct alignas(16) IItemRangeBinding
 {
+	// Try living without virtual destructor and trivially destructible members
 	virtual void ReadItems(FSaveRangeContext& Ctx) const = 0;
 	virtual void MakeItems(FLoadRangeContext& Ctx) const = 0;
 };
@@ -479,7 +480,7 @@ class FLeafRangeAllocator
 	FBuiltRange*				Range = nullptr;
 	const FUnpackedLeafType		Expected;
 
-	void* Allocate(FUnpackedLeafType Type, uint64 Num);
+	PLAINPROPS_API void* Allocate(uint64 Num, SIZE_T LeafSize);
 
 public:
 	FLeafRangeAllocator(FScratchAllocator& InScratch, FUnpackedLeafType InExpected) : Scratch(InScratch), Expected(InExpected) {}
@@ -488,39 +489,77 @@ public:
 	T* AllocateRange(SizeType Num)
 	{
 		check(ReflectArithmetic<T> == Expected);
-		return Num ? static_cast<T*>(Allocate(ReflectArithmetic<T>, IntCastChecked<uint64>(Num))) : nullptr;
+		return Num ? static_cast<T*>(Allocate(static_cast<uint64>(Num), sizeof(T))) : nullptr;
+	}
+
+	template<typename SizeType>
+	void* AllocateNonEmptyRange(SizeType Num, ELeafWidth Width)
+	{
+		check(Num > 0);
+		check(Width == Expected.Width);
+		return Allocate(static_cast<uint64>(Num), SizeOf(Width));
 	}
 
 	FBuiltRange* GetAllocatedRange() { return Range; }
 };
 
+// Hides internal representation to enable future format changes, 
+// e.g. store zeroes or 1.0f in some compact fashion or var int encodings
 class FLeafRangeLoadView
 {
 	const void*				Data;
-	uint64					Num;
+	uint64					NumItems;
 	FUnpackedLeafType		Leaf;
 
+	static constexpr bool CopyCompatible(ELeafType Dst, ELeafType Src)
+	{
+		return Dst == Src || (Dst == ELeafType::IntU && (Src == ELeafType::Hex || Src == ELeafType::Enum));
+	}
+
 public:
 	FLeafRangeLoadView(const void* InData, uint64 InNum, FUnpackedLeafType InLeaf)
 	: Data(InData)
-	, Num(InNum)
+	, NumItems(InNum)
 	, Leaf(InLeaf)
 	{}
+
+	uint64 Num() const { return NumItems; }
 	
-	// The returned ranges hide the internal representations so we can change format in the future, 
-	// e.g. store zeroes or 1.0f in some compact fashion or even var int encodings
 	template<typename T>
 	auto As() const
 	{
 		check(Leaf == ReflectLeaf<T>);
-		return TRangeView<T>(static_cast<const T*>(Data), Num);
+		return TRangeView<T>(static_cast<const T*>(Data), NumItems);
 	}
 
 	template<>
 	auto As<bool>() const
 	{
 		check(Leaf.Type == ELeafType::Bool && Leaf.Width == ELeafWidth::B8);
-		return FBoolRangeView(static_cast<const uint8*>(Data), Num);
+		return FBoolRangeView(static_cast<const uint8*>(Data), NumItems);
+	}
+
+	template<ELeafType Type, ELeafWidth Width>
+	void Copy(void* Dst, uint64 NumBytes) const
+	{
+		check(Width == Leaf.Width && CopyCompatible(Type, Leaf.Type));
+		check(NumBytes == NumItems * SizeOf(Leaf.Width));
+		if (NumBytes)
+		{
+			FMemory::Memcpy(Dst, Data, NumBytes);
+		}
+	}
+
+	template<>
+	void Copy<ELeafType::Bool, ELeafWidth::B8>(void* Dst, uint64 NumBytes) const
+	{
+		check(Leaf == ReflectArithmetic<bool>);
+		check(NumBytes == NumItems);
+		bool* bOut = static_cast<bool*>(Dst);
+		for (bool b : As<bool>())
+		{
+			*bOut++ = b;
+		}	
 	}
 };
 
@@ -529,7 +568,13 @@ struct alignas(16) ILeafRangeBinding
 {
 	virtual void	SaveLeaves(const void* Range, FLeafRangeAllocator& Out) const = 0;
 	virtual void	LoadLeaves(void* Range, FLeafRangeLoadView Leaves) const = 0;
-	virtual int64	DiffLeaves(const void* RangeA, const void* RangeB) const = 0;
+	virtual bool	DiffLeaves(const void* RangeA, const void* RangeB) const = 0;
+
+	template<typename SizeType>
+	static bool Diff(SizeType NumA, SizeType NumB, const void* A, const void* B, SIZE_T ItemSize)
+	{
+		return NumA != NumB || (NumA > 0 && FMemory::Memcmp(A, B, static_cast<uint64>(NumA) * ItemSize) != 0);
+	}
 };
 
 ////////////////////////////////////////////////////////////////////////////////////////////////
diff --git a/Engine/Plugins/Experimental/PlainProps/Source/Public/PlainPropsUeCoreBindings.h b/Engine/Plugins/Experimental/PlainProps/Source/Public/PlainPropsUeCoreBindings.h
--- a/Engine/Plugins/Experimental/PlainProps/Source/Public/PlainPropsUeCoreBindings.h
+++ b/Engine/Plugins/Experimental/PlainProps/Source/Public/PlainPropsUeCoreBindings.h
@@ -198,20 +198,12 @@ struct FStringBinding : public ILeafRangeBinding
 		}
 	}
 
-	virtual int64 DiffLeaves(const void* RangeA, const void* RangeB) const override
+	virtual bool DiffLeaves(const void* RangeA, const void* RangeB) const override
 	{
 		const FString& A = *static_cast<const FString*>(RangeA);
 		const FString& B = *static_cast<const FString*>(RangeB);
-		int32 ALen = A.Len();
-		int32 BLen = B.Len();
-
-		if (int32 LenDiff = ALen - BLen)
-		{
-			return LenDiff;
-		}
-
 		// Case-sensitive comparison
-		return ALen ? FMemory::Memcmp(A.GetCharArray().GetData(), B.GetCharArray().GetData(), ALen * sizeof(TCHAR)) : 0;
+		return Diff(A.Len(), B.Len(), A.GetCharArray().GetData(), B.GetCharArray().GetData(), sizeof(TCHAR));
 	}
 };
 
diff --git a/Engine/Plugins/Experimental/PlainPropsUObject/Source/Private/PlainPropsUObjectBindings.cpp b/Engine/Plugins/Experimental/PlainPropsUObject/Source/Private/PlainPropsUObjectBindings.cpp
--- a/Engine/Plugins/Experimental/PlainPropsUObject/Source/Private/PlainPropsUObjectBindings.cpp
+++ b/Engine/Plugins/Experimental/PlainPropsUObject/Source/Private/PlainPropsUObjectBindings.cpp
@@ -4,6 +4,7 @@
 #include "Containers/PagedArray.h"
 #include "Misc/CommandLine.h"
 #include "Misc/CoreDelegates.h"
+#include "Misc/DefinePrivateMemberPtr.h"
 #include "UObject/AnsiStrProperty.h"
 #include "UObject/EnumProperty.h"
 #include "UObject/FieldPathProperty.h"
@@ -32,7 +33,6 @@
 , ScriptInterface(	Names.IndexStruct({CoreUObject, Names.MakeTypename("ScriptInterface")}))
 {}
 
-
 ////////////////////////////////////////////////////////////////////////////////////////////////
 
 FGlobals::FGlobals() 
@@ -47,6 +47,45 @@
 
 ////////////////////////////////////////////////////////////////////////////////////////////////
 
+static constexpr uint64 LeafMask =			CASTCLASS_FNumericProperty | CASTCLASS_FEnumProperty| CASTCLASS_FBoolProperty;
+static constexpr uint64 IntSMask =			CASTCLASS_FInt8Property | CASTCLASS_FInt16Property | CASTCLASS_FIntProperty | CASTCLASS_FInt64Property;
+static constexpr uint64 IntUMask =			CASTCLASS_FByteProperty | CASTCLASS_FUInt16Property | CASTCLASS_FUInt32Property | CASTCLASS_FUInt64Property;
+static constexpr uint64 ContainerMask =		CASTCLASS_FArrayProperty | CASTCLASS_FSetProperty| CASTCLASS_FMapProperty  | CASTCLASS_FOptionalProperty;
+static constexpr uint64 StructuralMask =	CASTCLASS_FStructProperty | CASTCLASS_FDelegateProperty | CASTCLASS_FMulticastDelegateProperty;
+static constexpr uint64 StringMask =		CASTCLASS_FStrProperty | CASTCLASS_FUtf8StrProperty | CASTCLASS_FAnsiStrProperty;
+static constexpr uint64 CommonStructMask =	CASTCLASS_FNameProperty | CASTCLASS_FTextProperty | CASTCLASS_FInterfaceProperty | CASTCLASS_FFieldPathProperty |
+											CASTCLASS_FObjectProperty | CASTCLASS_FWeakObjectProperty | CASTCLASS_FSoftObjectProperty | CASTCLASS_FLazyObjectProperty;
+
+static FStructSchemaId FlagsToCommonBindId(uint64 MaskedCastFlags)
+{
+	switch (MaskedCastFlags)
+	{
+		case CASTCLASS_FNameProperty:			return GUE.CommonIds.Name;
+		case CASTCLASS_FTextProperty:			return GUE.CommonIds.Text;
+		case CASTCLASS_FInterfaceProperty:		return GUE.CommonIds.ScriptInterface;
+		case CASTCLASS_FFieldPathProperty:		return GUE.CommonIds.FieldPath;
+		case CASTCLASS_FObjectProperty:			return GUE.CommonIds.ObjectPtr;
+		case CASTCLASS_FWeakObjectProperty:		return GUE.CommonIds.WeakObjectPtr;
+		case CASTCLASS_FSoftObjectProperty:		return GUE.CommonIds.SoftObjectPtr;
+		case CASTCLASS_FLazyObjectProperty:		return GUE.CommonIds.LazyObjectPtr;
+		default:								break; // error
+	}
+
+	check(MaskedCastFlags); // @pre violated
+	check((MaskedCastFlags & CommonStructMask) == MaskedCastFlags); // @pre violated
+	checkf(FMath::CountBits(MaskedCastFlags) == 1, TEXT("Masked CASTCLASS flags %llx match more than one common property type"), MaskedCastFlags);
+	check(false); // Mismatch between this function and BindAsStructMask
+	return {};
+}
+
+template<uint64 Mask>
+static bool HasAny(uint64 Flags)
+{
+	return (Mask & Flags) != 0;
+}
+
+////////////////////////////////////////////////////////////////////////////////////////////////
+
 static FTypeId IndexType(const UField* Field)
 {	
 	FTypenameId Name = GUE.Names.MakeTypename(Field->GetFName());
@@ -142,10 +181,277 @@
 
 ////////////////////////////////////////////////////////////////////////////////////////////////
 
-struct FStaticArrayItemBinding : IItemRangeBinding
+// Helps cache array property range bindings
+struct FArrayPropertyInfo
 {
-	FStaticArrayItemBinding(uint32 InNum, uint32 InStride) : Num(InNum), Stride(InStride) {}
+	explicit FArrayPropertyInfo(FArrayProperty* Property)
+	: bFreezable(!!(Property->ArrayFlags & EArrayPropertyFlags::UsesMemoryImageAllocator))
+	, bDestructor(!(Property->Inner->PropertyFlags & (CPF_IsPlainOldData | CPF_NoDestructor)))
+	, bConstructor(!(Property->Inner->PropertyFlags & CPF_ZeroConstructor))
+	, ItemAlign(Property->Inner->GetMinAlignment())
+	, ItemSize(Property->Inner->GetElementSize())
+	{
+		check(ItemAlign == Property->Inner->GetMinAlignment()); // todo remove
+	}
 
+	union
+	{
+		struct
+		{
+			uint32	bFreezable : 1;
+			uint32	bDestructor : 1;	
+			uint32	bConstructor : 1;
+			uint32	ItemAlign : 29;
+		};
+		uint32		Int;
+	};
+	uint32			ItemSize;
+
+	bool			IsTrivial() const						{ return !bDestructor && !bConstructor; }
+	bool			operator==(FArrayPropertyInfo O) const	{ return Int == O.Int && ItemSize == O.ItemSize; }
+	friend uint32	GetTypeHash(FArrayPropertyInfo I)		{ return HashCombineFast(I.Int, I.ItemSize); };
+};
+static_assert(sizeof(FArrayPropertyInfo) == 8);
+
+// Cacheable FArrayProperty binding
+template<class ScriptArray>
+struct TTrivialArrayBinding : IItemRangeBinding
+{
+	TTrivialArrayBinding(FArrayPropertyInfo InFo) : Info(InFo) {}
+
+	FArrayPropertyInfo Info;
+
+	virtual void ReadItems(FSaveRangeContext& Ctx) const override
+	{
+		const ScriptArray& Array = Ctx.Request.GetRange<ScriptArray>();
+		Ctx.Items.SetAll({Array.GetData(), static_cast<uint64>(Array.Num())}, Info.ItemSize);
+	}
+
+	virtual void MakeItems(FLoadRangeContext& Ctx) const override
+	{
+		ScriptArray& Array = Ctx.Request.GetRange<ScriptArray>();
+
+		int32 NewNum = static_cast<int32>(Ctx.Request.NumTotal());
+		int32 Diff = NewNum - Array.Num();
+		if (Diff > 0)
+		{
+			Array.AddZeroed(Diff, Info.ItemSize, Info.ItemAlign);
+		}
+		else if (Diff < 0)
+		{
+			Array.Remove(NewNum, -Diff, Info.ItemSize, Info.ItemAlign);
+		}
+
+		Ctx.Items.Set(Array.GetData(), static_cast<uint64>(NewNum), Info.ItemSize);
+	}
+};
+
+// Temp hack. Long-term, try throwing away FProperty and implement Construct/DestroyValue replacement
+} // namespace PlainProps::UE
+UE_DEFINE_PRIVATE_MEMBER_PTR(void(void*) const, GInitPropertyValue, FProperty, InitializeValueInternal);
+UE_DEFINE_PRIVATE_MEMBER_PTR(void(void*) const, GDestroyPropertyValue, FProperty, DestroyValueInternal);
+namespace PlainProps::UE {
+
+// Currently can't extract constructor/destructor function pointers from FProperty, which
+// requires keeping FProperty* and prevents range binding reuse, @see AllocateArrayBinding()
+template<class ScriptArray>
+struct TNonTrivialArrayBinding : TTrivialArrayBinding<ScriptArray>
+{
+	using TTrivialArrayBinding<ScriptArray>::Info;
+
+	TNonTrivialArrayBinding(FArrayPropertyInfo InFo, FProperty* InNer)
+	: TTrivialArrayBinding<ScriptArray>(InFo)
+	, Inner(InNer)
+	{}
+
+	virtual void MakeItems(FLoadRangeContext& Ctx) const override
+	{
+		ScriptArray& Array = Ctx.Request.GetRange<ScriptArray>();
+
+		int32 NewNum = static_cast<int32>(Ctx.Request.NumTotal());
+		int32 Diff = NewNum - Array.Num();
+		if (Diff > 0)
+		{
+			int32 OldNum = Array.Add(Diff, Info.ItemSize, Info.ItemAlign);
+			ConstructItems(Diff, static_cast<uint8*>(Array.GetData()) + NumBytes(OldNum));
+		}
+		else if (Diff < 0)
+		{
+			DestructItems(-Diff, static_cast<uint8*>(Array.GetData()) + NumBytes(NewNum));
+			Array.Remove(NewNum, -Diff, Info.ItemSize, Info.ItemAlign);
+		}
+		
+		Ctx.Items.Set(Array.GetData(), static_cast<uint64>(NewNum), Info.ItemSize);
+	}
+
+	void ConstructItems(int32 Num, uint8* Items) const
+	{
+		if (Info.bConstructor)
+		{
+			for (uint8* It = Items, *End = Items + NumBytes(Num); It != End; It += Info.ItemSize)
+			{
+				((*Inner).*GInitPropertyValue)(It);
+			}
+		}
+		else
+		{
+			FMemory::Memzero(Items, NumBytes(Num));
+		}
+	}
+
+	void DestructItems(int32 Num, uint8* Items) const
+	{
+		if (Info.bDestructor)
+		{
+			for (uint8* It = Items, *End = Items + NumBytes(Num); It != End; It += Info.ItemSize)
+			{
+				((*Inner).*GDestroyPropertyValue)(It);
+			}
+		}
+	}
+
+	uint64 NumBytes(int32 NumItems) const
+	{
+		return static_cast<uint64>(NumItems) * Info.ItemSize;
+	}
+
+	FProperty* Inner;
+};
+
+template<ELeafType Type, ELeafWidth Width>
+struct TLeafArrayBinding : ILeafRangeBinding
+{
+	inline static constexpr SIZE_T LeafSize = SizeOf(Width);
+
+	virtual void SaveLeaves(const void* Range, FLeafRangeAllocator& Out) const override
+	{
+		const FScriptArray& Array = *static_cast<const FScriptArray*>(Range);
+		if (int32 Num = Array.Num())
+		{
+			void* Dst = Out.AllocateNonEmptyRange(Num, Width);
+			FMemory::Memcpy(Dst, Array.GetData(), NumBytes(Num));
+		}
+	}
+
+	virtual void LoadLeaves(void* Range, FLeafRangeLoadView Leaves) const override
+	{
+		FScriptArray& Array = *static_cast<FScriptArray*>(Range);
+		Array.SetNumUninitialized(static_cast<int32>(Leaves.Num()), LeafSize, LeafSize);
+		Leaves.Copy<Type, Width>(Array.GetData(), NumBytes(Array.Num()));
+	}
+
+	virtual bool DiffLeaves(const void* RangeA, const void* RangeB) const override
+	{
+		const FScriptArray& A = *static_cast<const FScriptArray*>(RangeA);
+		const FScriptArray& B = *static_cast<const FScriptArray*>(RangeB);
+		return Diff(A.Num(), B.Num(), A.GetData(), B.GetData(), LeafSize);
+	}
+
+	inline constexpr uint64 NumBytes(int32 NumItems) const
+	{
+		return static_cast<uint64>(NumItems) * LeafSize;
+	}
+};
+
+// Reusable cache of FArrayProperty range bindings
+class FArrayPropertyBindings
+{
+	static constexpr ERangeSizeType SizeType = RangeSizeOf(FDefaultAllocator::SizeType{});
+	
+	const TLeafArrayBinding<ELeafType::Bool, ELeafWidth::B8>	Bool;
+	const TLeafArrayBinding<ELeafType::Float, ELeafWidth::B32>	Float;
+	const TLeafArrayBinding<ELeafType::Float, ELeafWidth::B64>	Double;
+	const TLeafArrayBinding<ELeafType::IntS, ELeafWidth::B8>	IntS8;
+	const TLeafArrayBinding<ELeafType::IntS, ELeafWidth::B16>	IntS16;
+	const TLeafArrayBinding<ELeafType::IntS, ELeafWidth::B32>	IntS32;
+	const TLeafArrayBinding<ELeafType::IntS, ELeafWidth::B64>	IntS64;
+	const TLeafArrayBinding<ELeafType::IntU, ELeafWidth::B8>	IntU8;
+	const TLeafArrayBinding<ELeafType::IntU, ELeafWidth::B16>	IntU16;
+	const TLeafArrayBinding<ELeafType::IntU, ELeafWidth::B32>	IntU32;
+	const TLeafArrayBinding<ELeafType::IntU, ELeafWidth::B64>	IntU64;
+	const FRangeBinding											Integers[2][4];
+
+	TMap<FArrayPropertyInfo, IItemRangeBinding*>				Others;
+
+public:
+	FArrayPropertyBindings()
+	: Integers{{FRangeBinding(IntU8, SizeType),	FRangeBinding(IntU16, SizeType), FRangeBinding(IntU32, SizeType), FRangeBinding(IntU64, SizeType)},
+			   {FRangeBinding(IntS8, SizeType),	FRangeBinding(IntS16, SizeType), FRangeBinding(IntS32, SizeType), FRangeBinding(IntS64, SizeType)}}
+	{}
+
+	~FArrayPropertyBindings()
+	{
+		for (TPair<FArrayPropertyInfo, IItemRangeBinding*>& Cached : Others)
+		{
+			FMemory::Free(Cached.Value);
+		}
+	}
+
+	FRangeBinding RangeBind(FArrayPropertyInfo Info, uint64 InnerCastFlags)
+	{
+		if (HasAny<LeafMask>(InnerCastFlags) && !Info.bFreezable)
+		{
+			uint32 SizeIdx = FMath::FloorLog2NonZero(Info.ItemSize);
+			check(SizeIdx < 4);
+
+			// Note that we throw away enum schema, only size not needed to load/save enums
+			if (HasAny<IntSMask | IntUMask | CASTCLASS_FEnumProperty>(InnerCastFlags))
+			{
+				return Integers[HasAny<IntSMask>(InnerCastFlags)][SizeIdx];
+			}
+			
+			check(HasAny<CASTCLASS_FFloatProperty | CASTCLASS_FDoubleProperty | CASTCLASS_FBoolProperty >(InnerCastFlags));
+			const ILeafRangeBinding& Binding = HasAny<CASTCLASS_FBoolProperty>(InnerCastFlags) 
+				? (HasAny<CASTCLASS_FFloatProperty>(InnerCastFlags) ? static_cast<const ILeafRangeBinding&>(Float) : Double) : Bool;
+			return FRangeBinding(Binding, SizeType);
+		}
+		else if (IItemRangeBinding** Cached = Others.Find(Info))
+		{
+			return FRangeBinding(**Cached, SizeType);
+		}
+		
+		IItemRangeBinding* New = Info.bFreezable	? CreateAndCache<FFreezableScriptArray>(Info)
+													: CreateAndCache<FScriptArray>(Info);
+		return FRangeBinding(*New, SizeType);
+	}
+
+	template<typename ScriptArray>
+	IItemRangeBinding* CreateAndCache(FArrayPropertyInfo Info)
+	{
+		static_assert(std::is_trivially_destructible_v<TTrivialArrayBinding<ScriptArray>>);
+		return Others.Emplace(Info, new TTrivialArrayBinding<ScriptArray>(Info));
+	}
+};
+	
+static FArrayPropertyBindings GCachedArrayBindings;
+
+template<typename ScriptArray>
+IItemRangeBinding* CreateAndLeak(FArrayPropertyInfo Info, FProperty* Inner)
+{
+	return new TNonTrivialArrayBinding<ScriptArray>(Info, Inner);
+}
+
+static FRangeBinding AllocateArrayBinding(FArrayProperty* Property)
+{
+	FProperty* Inner = Property->Inner;
+	FArrayPropertyInfo Info(Property);
+	if (Info.IsTrivial())
+	{
+		return GCachedArrayBindings.RangeBind(Info, Inner->GetCastFlags());
+	}
+
+	// Todo: Ownership / memory leak, try make non-trivial case cacheable by making FProperty ctor/dtor extractable
+	IItemRangeBinding* Out = Info.bFreezable	? CreateAndLeak<FFreezableScriptArray>(Info, Inner)
+												: CreateAndLeak<FScriptArray>(Info, Inner);
+	return FRangeBinding(*Out, ERangeSizeType::S32);
+}
+
+////////////////////////////////////////////////////////////////////////////////////////////////
+
+struct FStaticArrayBinding : IItemRangeBinding
+{
+	FStaticArrayBinding(uint32 InNum, uint32 InStride) : Num(InNum), Stride(InStride) {}
+
 	uint32 Num;
 	uint32 Stride;
 
@@ -154,41 +460,13 @@
 		Ctx.Items.SetAll({Ctx.Request.Range, Num}, Stride);
 	}
 
-	virtual void MakeItems(FLoadRangeContext& Ctx) const override {}
+	virtual void MakeItems(FLoadRangeContext& Ctx) const override
+	{
+		Ctx.Items.Set(&Ctx.Request.GetRange<uint8>(), Num, Stride);
+	}
 };
-
 
 ////////////////////////////////////////////////////////////////////////////////////////////////
 
-static constexpr uint64 CommonStructMask = 
-	CASTCLASS_FNameProperty | CASTCLASS_FTextProperty | CASTCLASS_FInterfaceProperty | CASTCLASS_FFieldPathProperty |
-	CASTCLASS_FObjectProperty | CASTCLASS_FWeakObjectProperty | CASTCLASS_FSoftObjectProperty | CASTCLASS_FLazyObjectProperty;
-
-static FStructSchemaId FlagsToCommonBindId(uint64 MaskedCastFlags)
-{
-	switch (MaskedCastFlags)
-	{
-		case CASTCLASS_FNameProperty:			return GUE.CommonIds.Name;
-		case CASTCLASS_FTextProperty:			return GUE.CommonIds.Text;
-		case CASTCLASS_FInterfaceProperty:		return GUE.CommonIds.ScriptInterface;
-		case CASTCLASS_FFieldPathProperty:		return GUE.CommonIds.FieldPath;
-		case CASTCLASS_FObjectProperty:			return GUE.CommonIds.ObjectPtr;
-		case CASTCLASS_FWeakObjectProperty:		return GUE.CommonIds.WeakObjectPtr;
-		case CASTCLASS_FSoftObjectProperty:		return GUE.CommonIds.SoftObjectPtr;
-		case CASTCLASS_FLazyObjectProperty:		return GUE.CommonIds.LazyObjectPtr;
-		default:								break; // error
-	}
-
-	check(MaskedCastFlags); // @pre violated
-	check((MaskedCastFlags & CommonStructMask) == MaskedCastFlags); // @pre violated
-	checkf(FMath::CountBits(MaskedCastFlags) == 1, TEXT("Masked CASTCLASS flags %llx match more than one common property type"), MaskedCastFlags);
-	check(false); // Mismatch between this function and BindAsStructMask
-	return {};
-}
-
-////////////////////////////////////////////////////////////////////////////////////////////////
-
-static void BindStruct(FStructSchemaId Id, UStruct* Struct);
-
 class FPropertyBinder
 {
@@ -228,12 +506,6 @@
 		return MakeArrayView(OutData, OutNum);
 	}
 
-	template<uint64 Mask>
-	static bool HasAny(uint64 Flags)
-	{
-		return (Mask & Flags) != 0;
-	}
-
 	static FMemberBinding Todo(FProperty* Property)
 	{
 		return FMemberBinding(Property->GetOffset_ForInternal());
@@ -260,6 +532,29 @@
 		return OwnerOccupancy.GetValue();
 	}
 
+	inline FMemberBinding BindAsRange(FProperty* Property, FRangeBinding RangeBinding, FMemberBinding Inner)
+	{
+		FMemberBinding Out(Property->GetOffset_ForInternal());
+		Out.InnermostType = Inner.InnermostType;
+		Out.InnermostSchema = Inner.InnermostSchema;
+		Out.RangeBindings = AllocateRangeBindings(RangeBinding, Inner.RangeBindings);
+		return Out;
+	}
 
+	inline FMemberBinding BindAsStruct(FProperty* Property, FStructSchemaId Id)
+	{
+		FMemberBinding Out(Property->GetOffset_ForInternal());
+		Out.InnermostSchema = Id;
+		Out.InnermostType = DefaultStructBindType;
+		return Out;
+	}
+
+	inline FMemberBinding BindAsStruct(FProperty* Property, UStruct* Struct)
+	{
+		check(Struct);
+		FTypeId Type = IndexType(Struct);
+		return BindAsStruct(Property, GUE.Names.IndexStruct(Type));
+	}
+
 	inline static FBitfieldBoolBindType MakeBitfieldBool(uint8 BitIdx)
 	{
@@ -315,8 +610,11 @@
 		
 		return Out;
 	}
-
-	inline FMemberBinding BindArray(FArrayProperty* Property) { return Todo(Property); }
+	
+	inline FMemberBinding BindArray(FArrayProperty* Property)
+	{
+		return BindAsRange(Property, AllocateArrayBinding(Property), BindSingleProperty(Property->Inner));
+	}
 	inline FMemberBinding BindMap(FMapProperty* Property) { return Todo(Property); }
 	inline FMemberBinding BindSet(FSetProperty* Property) { return Todo(Property); }
 	inline FMemberBinding BindOptional(FOptionalProperty* Property) { return Todo(Property); }
@@ -329,28 +627,8 @@
 	inline FMemberBinding BindVRestValue(FVRestValueProperty* Property) { return Todo(Property); }
 	#endif
 
-	inline FMemberBinding BindAsStruct(FProperty* Property, FStructSchemaId Id)
-	{
-		FMemberBinding Out(Property->GetOffset_ForInternal());
-		Out.InnermostSchema = Id;
-		Out.InnermostType = DefaultStructBindType;
-		return Out;
-	}
-
-	inline FMemberBinding BindAsStruct(FProperty* Property, UStruct* Struct)
-	{
-		check(Struct);
-		FTypeId Type = IndexType(Struct);
-		return BindAsStruct(Property, GUE.Names.IndexStruct(Type));
-	}
-
 	FMemberBinding BindSingleProperty(FProperty* Property)
 	{
-		constexpr uint64 LeafMask = CASTCLASS_FNumericProperty | CASTCLASS_FEnumProperty| CASTCLASS_FBoolProperty;
-		constexpr uint64 ContainerMask = CASTCLASS_FArrayProperty | CASTCLASS_FSetProperty| CASTCLASS_FMapProperty  | CASTCLASS_FOptionalProperty;
-		constexpr uint64 StructuralMask = CASTCLASS_FStructProperty | CASTCLASS_FDelegateProperty | CASTCLASS_FMulticastDelegateProperty;
-		constexpr uint64 StringMask = CASTCLASS_FStrProperty | CASTCLASS_FUtf8StrProperty | CASTCLASS_FAnsiStrProperty;
-
 		uint64 Flags = Property->GetCastFlags();
 		if (HasAny<LeafMask>(Flags))
 		{
@@ -445,7 +723,7 @@
 			// and implement delta serialization manually
 
 			// Todo: Ownership / memory leak
-			const FStaticArrayItemBinding& ItemBinding = *new FStaticArrayItemBinding(ArrayDim, ElementSize);
+			const FStaticArrayBinding& ItemBinding = *new FStaticArrayBinding(ArrayDim, ElementSize);
 			ERangeSizeType SizeType = ArrayDim < 256 ? ERangeSizeType::U8 : ((ArrayDim < 65536) ? ERangeSizeType::U16 : ERangeSizeType::U32);
 			Out.RangeBindings =	AllocateRangeBindings(FRangeBinding(ItemBinding, SizeType), Out.RangeBindings);
 		}
