PlainProps add FSetProperty non-delta binding

#rnx

diff --git a/Engine/Plugins/Experimental/PlainPropsUObject/Source/Private/PlainPropsUObjectBindings.cpp b/Engine/Plugins/Experimental/PlainPropsUObject/Source/Private/PlainPropsUObjectBindings.cpp
--- a/Engine/Plugins/Experimental/PlainPropsUObject/Source/Private/PlainPropsUObjectBindings.cpp
+++ b/Engine/Plugins/Experimental/PlainPropsUObject/Source/Private/PlainPropsUObjectBindings.cpp
@@ -181,13 +181,59 @@ static FEnumSchemaId DeclareEnum(UEnum* Enum)
 
 ////////////////////////////////////////////////////////////////////////////////////////////////
 
+inline uint64 NumBytes(int32 NumItems, SIZE_T ItemSize)
+{
+	return static_cast<uint64>(NumItems) * ItemSize;
+}
+
+inline bool HasConstructor(FProperty* Property)
+{
+	return !(Property->PropertyFlags & CPF_ZeroConstructor);
+}
+
+inline bool HasDestructor(FProperty* Property)
+{
+	return !(Property->PropertyFlags & (CPF_IsPlainOldData | CPF_NoDestructor));
+}
+
+// Temp hack. Long-term, try throwing away FProperty and implement Construct/DestroyValue replacement
+} // namespace PlainProps::UE
+UE_DEFINE_PRIVATE_MEMBER_PTR(void(void*) const, GInitPropertyValue, FProperty, InitializeValueInternal);
+UE_DEFINE_PRIVATE_MEMBER_PTR(void(void*) const, GDestroyPropertyValue, FProperty, DestroyValueInternal);
+namespace PlainProps::UE {
+
+inline void InitializeValues(FProperty* Property, uint8* Items, int32 Num, SIZE_T Stride)
+{
+	if (Property)
+	{
+		for (uint8* It = Items, *End = Items + NumBytes(Num, Stride); It != End; It += Stride)
+		{
+			((*Property).*GInitPropertyValue)(It);
+		}
+	}
+	else
+	{
+		FMemory::Memzero(Items, NumBytes(Num, Stride));
+	}
+}
+
+inline void DestroyValues(FProperty* Property, uint8* Items, int32 Num, SIZE_T Stride)
+{
+	for (uint8* It = Items, *End = Items + NumBytes(Num, Stride); It != End; It += Stride)
+	{
+		((*Property).*GDestroyPropertyValue)(It);
+	}
+}
+
+////////////////////////////////////////////////////////////////////////////////////////////////
+
 // Helps cache array property range bindings
 struct FArrayPropertyInfo
 {
 	explicit FArrayPropertyInfo(FArrayProperty* Property)
 	: bFreezable(!!(Property->ArrayFlags & EArrayPropertyFlags::UsesMemoryImageAllocator))
-	, bDestructor(!(Property->Inner->PropertyFlags & (CPF_IsPlainOldData | CPF_NoDestructor)))
-	, bConstructor(!(Property->Inner->PropertyFlags & CPF_ZeroConstructor))
+	, bDestructor(HasDestructor(Property->Inner))
+	, bConstructor(HasConstructor(Property->Inner))
 	, ItemAlign(Property->Inner->GetMinAlignment())
 	, ItemSize(Property->Inner->GetElementSize())
 	{
@@ -246,12 +292,6 @@ struct TTrivialArrayBinding : IItemRangeBinding
 	}
 };
 
-// Temp hack. Long-term, try throwing away FProperty and implement Construct/DestroyValue replacement
-} // namespace PlainProps::UE
-UE_DEFINE_PRIVATE_MEMBER_PTR(void(void*) const, GInitPropertyValue, FProperty, InitializeValueInternal);
-UE_DEFINE_PRIVATE_MEMBER_PTR(void(void*) const, GDestroyPropertyValue, FProperty, DestroyValueInternal);
-namespace PlainProps::UE {
-
 // Currently can't extract constructor/destructor function pointers from FProperty, which
 // requires keeping FProperty* and prevents range binding reuse, @see AllocateArrayBinding()
 template<class ScriptArray>
@@ -273,11 +313,11 @@ struct TNonTrivialArrayBinding : TTrivialArrayBinding<ScriptArray>
 		if (Diff > 0)
 		{
 			int32 OldNum = Array.Add(Diff, Info.ItemSize, Info.ItemAlign);
-			ConstructItems(Diff, static_cast<uint8*>(Array.GetData()) + NumBytes(OldNum));
+			ConstructItems(Diff, static_cast<uint8*>(Array.GetData()) + NumBytes(OldNum, Info.ItemSize));
 		}
 		else if (Diff < 0)
 		{
-			DestructItems(-Diff, static_cast<uint8*>(Array.GetData()) + NumBytes(NewNum));
+			DestructItems(-Diff, static_cast<uint8*>(Array.GetData()) + NumBytes(NewNum, Info.ItemSize));
 			Array.Remove(NewNum, -Diff, Info.ItemSize, Info.ItemAlign);
 		}
 		
@@ -286,33 +326,13 @@ struct TNonTrivialArrayBinding : TTrivialArrayBinding<ScriptArray>
 
 	void ConstructItems(int32 Num, uint8* Items) const
 	{
-		if (Info.bConstructor)
-		{
-			for (uint8* It = Items, *End = Items + NumBytes(Num); It != End; It += Info.ItemSize)
-			{
-				((*Inner).*GInitPropertyValue)(It);
-			}
-		}
-		else
-		{
-			FMemory::Memzero(Items, NumBytes(Num));
-		}
+		InitializeValues(Info.bConstructor ? Inner : nullptr, Items, Num, Info.ItemSize);
 	}
 
 	void DestructItems(int32 Num, uint8* Items) const
 	{
-		if (Info.bDestructor)
-		{
-			for (uint8* It = Items, *End = Items + NumBytes(Num); It != End; It += Info.ItemSize)
-			{
-				((*Inner).*GDestroyPropertyValue)(It);
-			}
-		}
-	}
-
-	uint64 NumBytes(int32 NumItems) const
-	{
-		return static_cast<uint64>(NumItems) * Info.ItemSize;
+		int32 NumDestroy = Info.bDestructor * Num;
+		DestroyValues(Inner, Items, NumDestroy, Info.ItemSize);
 	}
 
 	FProperty* Inner;
@@ -448,6 +468,104 @@ static FRangeBinding AllocateArrayBinding(FArrayProperty* Property)
 
 ////////////////////////////////////////////////////////////////////////////////////////////////
 
+struct FSetBinding : IItemRangeBinding
+{
+	FProperty*					Element;
+	const FScriptSetLayout		Layout;
+	const bool					bDestruct;
+	const bool					bConstruct;
+	
+	explicit FSetBinding(FSetProperty* In)
+	: Element(In->ElementProp)
+	, Layout(In->SetLayout)
+	, bDestruct(HasDestructor(In->ElementProp))
+	, bConstruct(HasConstructor(In->ElementProp))
+	{}
+
+	virtual void MakeItems(FLoadRangeContext& Ctx) const override
+	{
+		FScriptSet& Set = Ctx.Request.GetRange<FScriptSet>();
+		const int32 NewNum = static_cast<int32>(Ctx.Request.NumTotal());
+				
+		if (Ctx.Request.IsFirstCall())
+		{
+			if (bDestruct && !Set.IsEmpty())
+			{
+				uint8* Data = static_cast<uint8*>(Set.GetData(0, Layout));
+				DestroyValues(Element, Data, Set.NumUnchecked(), Layout.Size);
+			}
+
+			Set.Empty(NewNum, Layout);
+			
+			if (NewNum > 0)
+			{
+				// Todo: Add TScriptSparseArray::AppendUninitialized or SetNumUninitialized
+				for (int32 Idx = 0; Idx < NewNum; ++Idx)
+				{
+					Set.AddUninitialized(Layout);
+				}
+				check(/* Set.IsCompact() */ Set.GetMaxIndex() == NewNum);
+
+				uint8* Data = static_cast<uint8*>(Set.GetData(0, Layout));
+				InitializeValues(bConstruct ? Element : nullptr, Data, NewNum, Layout.Size);
+
+				Ctx.Items.Set(Data, Ctx.Request.NumTotal(), Layout.Size);
+				Ctx.Items.RequestFinalCall();
+			}
+		}
+		else
+		{
+			check(Ctx.Request.IsFinalCall());
+			Set.Rehash(Layout, [this](const void* Src) { return Element->GetValueTypeHash(Src); });
+		}
+	}
+
+	virtual void ReadItems(FSaveRangeContext& Ctx) const override
+	{
+		const FScriptSparseArray& Elems = Ctx.Request.GetRange<FScriptSparseArray>();
+		const uint8* Data = static_cast<const uint8*>(Elems.GetData(0, Layout.SparseArrayLayout));
+
+		if (FExistingItemSlice LastRead = Ctx.Items.Slice)
+		{
+			// Continue partial response
+			int64 TotalBytesRead = static_cast<const uint8*>(LastRead.Data) - Data;
+			check(TotalBytesRead % Layout.Size == 0);
+			int32 NextIdx = /* skip known invalid */ 1 + TotalBytesRead / Layout.Size;
+			check(NextIdx < Elems.GetMaxIndex());
+			check(NextIdx > static_cast<int32>(LastRead.Num));
+			Ctx.Items.Slice = GetContiguousSlice(NextIdx, Elems, Data);
+		}
+		else if (Elems.IsCompact())
+		{
+			int32 Num = Elems.Num();
+			Ctx.Items.SetAll(Num ? Data : nullptr, Num);
+		}
+		else
+		{
+			// Start partial response
+			Ctx.Items.NumTotal = Elems.Num();
+			Ctx.Items.Stride = Layout.Size;
+			Ctx.Items.Slice = GetContiguousSlice(0, Elems, Data);
+		}
+	}
+
+	FExistingItemSlice GetContiguousSlice(int32 Idx, const FScriptSparseArray& Elems, const uint8* Data) const
+	{
+		int32 Num = 1;
+		for (;!Elems.IsValidIndex(Idx); ++Idx) {}
+		for (; Elems.IsValidIndex(Idx + Num); ++Num) {}
+		return { Data + NumBytes(Idx, Layout.Size), static_cast<uint64>(Num) };
+	}
+};
+
+static FRangeBinding AllocateSetBinding(FSetProperty* Property)
+{
+	// Todo: Ownership / memory leak
+	return FRangeBinding(*new FSetBinding(Property), ERangeSizeType::S32);
+}
+
+////////////////////////////////////////////////////////////////////////////////////////////////
+
 struct FStaticArrayBinding : IItemRangeBinding
 {
 	FStaticArrayBinding(uint32 InNum, uint32 InStride) : Num(InNum), Stride(InStride) {}
@@ -616,14 +734,17 @@ private:
 		return BindAsRange(Property, AllocateArrayBinding(Property), BindSingleProperty(Property->Inner));
 	}
 	inline FMemberBinding BindMap(FMapProperty* Property) { return Todo(Property); }
-	inline FMemberBinding BindSet(FSetProperty* Property) { return Todo(Property); }
+	inline FMemberBinding BindSet(FSetProperty* Property)
+	{ 
+		return BindAsRange(Property, AllocateSetBinding(Property), BindSingleProperty(Property->ElementProp));
+	}
 	inline FMemberBinding BindOptional(FOptionalProperty* Property) { return Todo(Property); }
 
 	template<typename StrProperty>
 	inline FMemberBinding BindStr(StrProperty* Property) { return Todo(Property); }
 
 	#if WITH_VERSE_VM
-	inline FMemberBinding BindVValue(FVRestValueProperty* Property) { return Todo(Property); }
+	inline FMemberBinding BindVValue(FVValueProperty* Property) { return Todo(Property); }
 	inline FMemberBinding BindVRestValue(FVRestValueProperty* Property) { return Todo(Property); }
 	#endif
 
