PlainProps add FMapProperty flat binding that generates TPair bindings + untyped parameter indexing

#rnx

diff --git a/Engine/Plugins/Experimental/PlainProps/Source/Private/PlainPropsBind.cpp b/Engine/Plugins/Experimental/PlainProps/Source/Private/PlainPropsBind.cpp
--- a/Engine/Plugins/Experimental/PlainProps/Source/Private/PlainPropsBind.cpp
+++ b/Engine/Plugins/Experimental/PlainProps/Source/Private/PlainPropsBind.cpp
@@ -230,6 +230,33 @@ TConstArrayView<FInnerStruct> FCustomBindings::FetchInners(FInnersHandle In) con
 
 ////////////////////////////////////////////////////////////////////////////////////////////////
 
+inline FTypeId IndexInnermost(FIdIndexerBase& Ids, FMemberBindType Type, FOptionalSchemaId Schema)
+{
+	if (Type.IsStruct())
+	{
+		return Ids.Resolve(static_cast<FStructSchemaId>(Schema.Get()));
+	}
+
+	FUnpackedLeafType Leaf = ToLeafType(Type.AsLeaf());
+	if (Leaf.Type == ELeafType::Enum)
+	{
+		return Ids.Resolve(static_cast<FEnumSchemaId>(Schema.Get()));
+	}
+
+	return Ids.MakeLeafParameter(Leaf);
+}
+
+FTypeId	FMemberBinding::IndexParameterName(FIdIndexerBase& Ids) const
+{
+	FTypeId Out = IndexInnermost(Ids, InnermostType, InnermostSchema);
+	for (FRangeBinding Range : ReverseIterate(RangeBindings))
+	{
+		FTypeId SizeType = Ids.MakeRangeParameter(Range.GetSizeType());
+		Out = Ids.MakeAnonymousParametricType({ Out, SizeType });
+	}
+	return Out;
+}
+
 FSchemaBindings::~FSchemaBindings()
 {}
 
diff --git a/Engine/Plugins/Experimental/PlainProps/Source/Private/PlainPropsPrint.cpp b/Engine/Plugins/Experimental/PlainProps/Source/Private/PlainPropsPrint.cpp
new file mode 100644
--- /dev/null
+++ b/Engine/Plugins/Experimental/PlainProps/Source/Private/PlainPropsPrint.cpp
@@ -0,0 +1,55 @@
+// Copyright Epic Games, Inc. All Rights Reserved.
+
+#include "PlainPropsPrint.h"
+#include "PlainPropsIndex.h"
+
+namespace PlainProps
+{
+
+struct FLiterals
+{
+	FAnsiStringView Ranges[9] = { 
+		("Uni"), ("S8"), ("U8"), ("S16"), ("U16"), ("S32"), ("U32"), ("S64"), ("U64") };
+	FAnsiStringView Leaves[8][4] = {
+		{("bool"),		("ERR_b16"),	("ERR_b32"),	("ERR_b64")},
+		{("i8"),		("i16"),		("i32"),		("i64")},
+		{("u8"),		("u16"),		("u32"),		("u64")},
+		{("ERR_fp8"),	("ERR_fp16"),	("float"),		("double")},
+		{("hex8"),		("hex16"),		("hex32"),		("hex64")},
+		{("enum8"),		("enum16"),		("enum32"),		("enum64")},
+		{("utf8"),		("utf16"),		("utf32"),		("ERR_utf64")},
+		{("ERR_oob"),	("ERR_oob"),	("ERR_oob"),	("ERR_oob")}};
+};
+static const FLiterals GLiterals;
+
+//////////////////////////////////////////////////////////////////////////
+
+FAnsiStringView ToString(ERangeSizeType Width)
+{
+	return GLiterals.Ranges[(uint8)Width];
+}
+
+FAnsiStringView ToString(FUnpackedLeafType Leaf)
+{
+	return GLiterals.Leaves[(uint8)Leaf.Type][(uint8)Leaf.Width];
+}
+
+//////////////////////////////////////////////////////////////////////////
+
+void FIdIndexerBase::InitParameterNames()
+{
+	for (uint32 T = 0; T < 8; ++T)
+	{
+		for (uint32 W = 0; W < 4; ++W)
+		{
+			Leaves[T][W] = InitParameterName(GLiterals.Leaves[T][W].GetData());
+		}
+	}
+
+	for (uint32 S = 0; S < 9; ++S)
+	{
+		Ranges[S] = InitParameterName(GLiterals.Ranges[S].GetData());
+	}
+}
+
+} // namespace PlainProps
\ No newline at end of file
diff --git a/Engine/Plugins/Experimental/PlainProps/Source/Public/PlainPropsBind.h b/Engine/Plugins/Experimental/PlainProps/Source/Public/PlainPropsBind.h
--- a/Engine/Plugins/Experimental/PlainProps/Source/Public/PlainPropsBind.h
+++ b/Engine/Plugins/Experimental/PlainProps/Source/Public/PlainPropsBind.h
@@ -116,6 +116,7 @@ union FMemberBindType
 	uint8					AsByte() const		{ return BitCast<uint8>(*this); }
 
 	friend inline bool operator==(FMemberBindType A, FMemberBindType B) { return A.AsByte() == B.AsByte(); }
+	friend inline uint32 GetTypeHash(FMemberBindType In) { return In.AsByte(); }
 private:
     EMemberKind				Kind : 2;
 	FLeafBindType			Leaf;
@@ -594,6 +595,8 @@ public:
 	const IItemRangeBinding&	AsItemBinding() const		{ check(!IsLeafBinding()); return *reinterpret_cast<IItemRangeBinding*>(Handle & BindMask); }
 	const ILeafRangeBinding&	AsLeafBinding() const		{ check( IsLeafBinding()); return *reinterpret_cast<ILeafRangeBinding*>(Handle & BindMask); }
 	ERangeSizeType				GetSizeType() const			{ return static_cast<ERangeSizeType>(Handle & SizeMask); }
+
+	inline bool	operator==(FRangeBinding O) const { return Handle == O.Handle; }
 };
 
 template<typename T>
@@ -615,6 +618,8 @@ struct FMemberBinding
 	FMemberBindType					InnermostType;		// Always Leaf or Struct
 	FOptionalSchemaId				InnermostSchema;	// Enum or struct schema
 	TConstArrayView<FRangeBinding>	RangeBindings;		// Non-empty -> Range
+
+	PLAINPROPS_API FTypeId			IndexParameterName(FIdIndexerBase& Ids) const;
 };
 
 class FSchemaBindings : public IStructBindIds
@@ -799,19 +804,19 @@ FTypeId IndexParameterName()
 	}
 	else
 	{
-		FTypeId ItemTypename = IndexParameterName<Ids, Kind, typename RangeBinding::ItemType>();
-		FTypeId SizeTypename = IndexArithmeticName<Ids, typename RangeBinding::SizeType>();
+		FTypeId ItemParam = IndexParameterName<Ids, Kind, typename RangeBinding::ItemType>();
+		FTypeId SizeParam = Ids::GetIndexer().MakeRangeParameter(RangeSizeOf(typename RangeBinding::SizeType{}));
 
 		if constexpr (Kind == ETypename::Decl)
 		{
 			// Type-erase range type
-			return Ids::GetIndexer().MakeAnonymousParametricType({ItemTypename, SizeTypename});
+			return Ids::GetIndexer().MakeAnonymousParametricType({ItemParam, SizeParam});
 		}
 		else
 		{
 			using Typename = TTypename<T>;
 			FTypeId RangeBindName = { IndexNamespaceId<Ids, Typename>(), Ids::IndexTypename(ToAnsiView(Typename::RangeBindName)) };
-			return Ids::GetIndexer().MakeParametricType(RangeBindName, {ItemTypename, SizeTypename});
+			return Ids::GetIndexer().MakeParametricType(RangeBindName, {ItemParam, SizeParam});
 		}
 	}
 }
diff --git a/Engine/Plugins/Experimental/PlainProps/Source/Public/PlainPropsIndex.h b/Engine/Plugins/Experimental/PlainProps/Source/Public/PlainPropsIndex.h
--- a/Engine/Plugins/Experimental/PlainProps/Source/Public/PlainPropsIndex.h
+++ b/Engine/Plugins/Experimental/PlainProps/Source/Public/PlainPropsIndex.h
@@ -54,6 +54,8 @@ public:
 	PLAINPROPS_API FParametricTypeId	MakeParametricTypeId(FOptionalConcreteTypenameId Name, TConstArrayView<FTypeId> Params);
 	PLAINPROPS_API FTypeId				MakeParametricType(FTypeId Type, TConstArrayView<FTypeId> Params);
 	PLAINPROPS_API FTypeId				MakeAnonymousParametricType(TConstArrayView<FTypeId> Params);
+	FTypeId								MakeLeafParameter(FUnpackedLeafType Leaf)			{ return { NoId, FTypenameId(Leaves[uint8(Leaf.Type)][uint8(Leaf.Width)]) }; }
+	FTypeId								MakeRangeParameter(ERangeSizeType SizeType)			{ return { NoId, FTypenameId(Ranges[uint8(SizeType)]) }; }
 
 	PLAINPROPS_API FEnumSchemaId		IndexEnum(FTypeId Type);
 	PLAINPROPS_API FStructSchemaId		IndexStruct(FTypeId Type);
@@ -75,6 +77,11 @@ protected:
 	FParametricTypeIndexer				ParametricTypes;
 	TSet<FTypeId>						Enums;
 	TSet<FTypeId>						Structs;
+	FConcreteTypenameId					Leaves[8][4];
+	FConcreteTypenameId					Ranges[9];
+
+	PLAINPROPS_API void					InitParameterNames();
+	virtual FConcreteTypenameId			InitParameterName(const char* Name) = 0;
 };
 
 template<class NameType>
@@ -83,27 +90,36 @@ void AppendString(FString& Out, const NameType& Str)
 	Out.Append(Str);
 }
 
-template<class NameType>
+template<class NameT>
 class TIdIndexer : public FIdIndexerBase
 {
 public:
-	template<typename T> FNameId		MakeName(T&& Str)									{ return { IntCastChecked<uint32>(Names.Add(NameType(Str)).AsInteger()) }; }
-	template<typename T> FMemberId		NameMember(T&& Name)								{ return { MakeName(Name) }; }
-	template<typename T> FFlatScopeId	MakeScope(T&& Str)									{ return static_cast<FFlatScopeId>(MakeName(Str)); }
-	template<typename T> FScopeId		NestScope(FScopeId Outer, T&& Inner)				{ return NestFlatScope(Outer, static_cast<FFlatScopeId>(MakeName(Inner))); }
-	template<typename T> FTypenameId	MakeTypename(T&& Name)								{ return FTypenameId(FConcreteTypenameId(MakeName(Name))); }
-	template<typename T> FTypeId		MakeType(T&& Scope, T&& Name)						{ return { FScopeId(MakeScope(Scope)), MakeTypename(Name) }; }
+	TIdIndexer() { InitParameterNames(); }
+
+	template<typename T> FNameId				MakeName(T&& Name)						{ return { static_cast<uint32>(Names.Add(NameT(Name)).AsInteger()) }; }
+	template<typename T> FMemberId				NameMember(T&& Name)					{ return { MakeName(Forward<T>(Name)) }; }
+	template<typename T> FConcreteTypenameId	NameType(T&& Name)						{ return { MakeName(Forward<T>(Name)) }; }
+	template<typename T> FFlatScopeId			NameScope(T&& Name)						{ return { MakeName(Forward<T>(Name)) }; }
+	template<typename T> FScopeId				MakeScope(T&& Name)						{ return FScopeId(NameScope(Forward<T>(Name))); }
+	template<typename T> FScopeId				NestScope(FScopeId Outer, T&& Inner)	{ return NestFlatScope(Outer, NameScope(Forward<T>(Inner))); }
+	template<typename T> FTypenameId			MakeTypename(T&& Name)					{ return FTypenameId(NameType(Forward<T>(Name))); }
+	template<typename T> FTypeId				MakeType(T&& Scope, T&& Name)			{ return { MakeScope(Forward<T>(Scope)), MakeTypename(Forward<T>(Name)) }; }
 
-	NameType							ResolveName(FNameId Id)	const						{ return Names.Get(FSetElementId::FromInteger(Id.Idx)); }
-	virtual uint32						NumNames() const override final						{ return static_cast<uint32>(Names.Num()); }
+	NameT										ResolveName(FNameId Id)	const			{ return Names.Get(FSetElementId::FromInteger(Id.Idx)); }
+	virtual uint32								NumNames() const override final			{ return static_cast<uint32>(Names.Num()); }
 
 protected:
-	TSet<NameType>						Names;
+	TSet<NameT>									Names;
 
-	virtual void						AppendDebugString(FString& Out, FNameId Id) const override final
+	virtual void								AppendDebugString(FString& Out, FNameId Id) const override final
 	{
 		AppendString(Out, ResolveName(Id));
 	}
+
+	virtual FConcreteTypenameId					InitParameterName(const char* Name) override final
+	{ 
+		return NameType(Name);
+	}
 };
 
 
diff --git a/Engine/Plugins/Experimental/PlainProps/Source/Public/PlainPropsPrint.h b/Engine/Plugins/Experimental/PlainProps/Source/Public/PlainPropsPrint.h
--- a/Engine/Plugins/Experimental/PlainProps/Source/Public/PlainPropsPrint.h
+++ b/Engine/Plugins/Experimental/PlainProps/Source/Public/PlainPropsPrint.h
@@ -8,25 +8,7 @@
 namespace PlainProps
 {
 
-inline FAnsiStringView ToString(ERangeSizeType Width)
-{
-	static const FAnsiStringView Widths[] = { "Uni", "S8", "U8", "S16", "U16", "S32", "U32", "S64", "U64" };
-	return Widths[(uint8)Width];
-}
-
-inline FAnsiStringView ToString(FUnpackedLeafType Leaf)
-{
-	static const FAnsiStringView Leaves[8][4] = {
-		{("bool"),		("ERR_b16"),	("ERR_b32"),	("ERR_b64")},
-		{("i8"),		("i16"),		("i32"),		("i64")},
-		{("u8"),		("u16"),		("u32"),		("u64")},
-		{("ERR_fp8"),	("ERR_fp16"),	("float"),		("double")},
-		{("hex8"),		("hex16"),		("hex32"),		("hex64")},
-		{("enum8"),		("enum16"),		("enum32"),		("enum64")},
-		{("utf8"),		("utf16"),		("utf32"),		("ERR_utf64")},
-		{("ERR_oob"),	("ERR_oob"),	("ERR_oob"),	("ERR_oob")}};
-
-	return Leaves[(uint8)Leaf.Type][(uint8)Leaf.Width];
-}
+PLAINPROPS_API FAnsiStringView ToString(ERangeSizeType Width);
+PLAINPROPS_API FAnsiStringView ToString(FUnpackedLeafType Leaf);
 
 } // namespace PlainProps
diff --git a/Engine/Plugins/Experimental/PlainProps/Source/Public/PlainPropsTypes.h b/Engine/Plugins/Experimental/PlainProps/Source/Public/PlainPropsTypes.h
--- a/Engine/Plugins/Experimental/PlainProps/Source/Public/PlainPropsTypes.h
+++ b/Engine/Plugins/Experimental/PlainProps/Source/Public/PlainPropsTypes.h
@@ -153,8 +153,8 @@ class FScopeId
 	uint32 Handle;
 public:
 	FScopeId(FNoId) : Handle(Unscoped) {}
-	FScopeId(FFlatScopeId Flat) : Handle(Flat.Name.Idx) 				{ check(AsFlat().Name == Flat.Name); }
-	FScopeId(FNestedScopeId Nested) : Handle(Nested.Idx | NestedBit) 	{ check(AsNested().Idx == Nested.Idx); }
+	explicit FScopeId(FFlatScopeId Flat) : Handle(Flat.Name.Idx) 				{ check(AsFlat().Name == Flat.Name); }
+	explicit FScopeId(FNestedScopeId Nested) : Handle(Nested.Idx | NestedBit) 	{ check(AsNested().Idx == Nested.Idx); }
 
 	explicit						operator bool() const			{ return Handle != ~0u; }
 	bool							IsFlat() const 					{ return !(Handle & NestedBit); }
diff --git a/Engine/Plugins/Experimental/PlainProps/Source/Tests/PlainPropsUeCoreTests.cpp b/Engine/Plugins/Experimental/PlainProps/Source/Tests/PlainPropsUeCoreTests.cpp
--- a/Engine/Plugins/Experimental/PlainProps/Source/Tests/PlainPropsUeCoreTests.cpp
+++ b/Engine/Plugins/Experimental/PlainProps/Source/Tests/PlainPropsUeCoreTests.cpp
@@ -33,7 +33,7 @@ struct FIds
 	static FNameId			IndexName(FAnsiStringView Name)				{ return GNames.MakeName(FName(Name)); }
 	static FMemberId		IndexMember(FAnsiStringView Name)			{ return GNames.NameMember(FName(Name)); }
 	static FTypenameId		IndexTypename(FAnsiStringView Name)			{ return GNames.MakeTypename(FName(Name)); }
-	static FFlatScopeId		IndexScope(FAnsiStringView Name)			{ return GNames.MakeScope(FName(Name)); }
+	static FScopeId			IndexScope(FAnsiStringView Name)			{ return GNames.MakeScope(FName(Name)); }
 	static FEnumSchemaId	IndexEnum(FTypeId Type)						{ return GNames.IndexEnum(Type); }
 	static FStructSchemaId	IndexStruct(FTypeId Type)					{ return GNames.IndexStruct(Type); }
 	static FIdIndexerBase&	GetIndexer()								{ return GNames; }
diff --git a/Engine/Plugins/Experimental/PlainPropsUObject/Source/Private/PlainPropsUObjectBindings.cpp b/Engine/Plugins/Experimental/PlainPropsUObject/Source/Private/PlainPropsUObjectBindings.cpp
--- a/Engine/Plugins/Experimental/PlainPropsUObject/Source/Private/PlainPropsUObjectBindings.cpp
+++ b/Engine/Plugins/Experimental/PlainPropsUObject/Source/Private/PlainPropsUObjectBindings.cpp
@@ -1,7 +1,9 @@
 // Copyright Epic Games, Inc. All Rights Reserved.
 
 #include "PlainPropsUObjectRuntime.h"
+#include "Algo/Compare.h"
 #include "Containers/PagedArray.h"
+#include "Hash/xxhash.h"
 #include "Misc/CommandLine.h"
 #include "Misc/CoreDelegates.h"
 #include "Misc/DefinePrivateMemberPtr.h"
@@ -21,7 +23,8 @@ namespace PlainProps::UE
 {
 
 FCommonUObjectIds::FCommonUObjectIds(TIdIndexer<FName>& Names)
-: CoreUObject(Names.MakeScope("/Script/CoreUObject"))
+: Core(				Names.MakeScope("/Script/Core"))
+, CoreUObject(		Names.MakeScope("/Script/CoreUObject"))
 , Name(				Names.IndexStruct({CoreUObject, Names.MakeTypename("Name")}))
 , Text(				Names.IndexStruct({CoreUObject, Names.MakeTypename("Text")}))
 , FieldPath(		Names.IndexStruct({CoreUObject, Names.MakeTypename("FieldPath")}))
@@ -31,6 +34,9 @@ FCommonUObjectIds::FCommonUObjectIds(TIdIndexer<FName>& Names)
 , LazyObjectPtr(	Names.IndexStruct({CoreUObject, Names.MakeTypename("LazyObjectPtr")}))
 , SoftObjectPtr(	Names.IndexStruct({CoreUObject, Names.MakeTypename("SoftObjectPtr")}))
 , ScriptInterface(	Names.IndexStruct({CoreUObject, Names.MakeTypename("ScriptInterface")}))
+, Pair(				Names.NameType("Pair"))
+, Key(				Names.NameMember("Key"))
+, Value(			Names.NameMember("Value"))
 {}
 
 ////////////////////////////////////////////////////////////////////////////////////////////////
@@ -94,7 +100,7 @@ static FTypeId IndexType(const UField* Field)
 	TArray<FFlatScopeId, TInlineAllocator<64>> ReversedOuters;
 	for (const UObject* Outer = Field->GetOuter(); Outer; Outer = Outer->GetOuter())
 	{
-		ReversedOuters.Add(GUE.Names.MakeScope(Outer->GetFName()));
+		ReversedOuters.Add(GUE.Names.NameScope(Outer->GetFName()));
 	}
 	
 	return { GUE.Names.NestReversedScopes(ReversedOuters), Name };
@@ -202,6 +208,12 @@ UE_DEFINE_PRIVATE_MEMBER_PTR(void(void*) const, GInitPropertyValue, FProperty, I
 UE_DEFINE_PRIVATE_MEMBER_PTR(void(void*) const, GDestroyPropertyValue, FProperty, DestroyValueInternal);
 namespace PlainProps::UE {
 
+
+inline void InitializeValue(FProperty* Property, void* Item)
+{
+	((*Property).*GInitPropertyValue)(Item);
+}
+
 inline void InitializeValues(FProperty* Property, uint8* Items, int32 Num, SIZE_T Stride)
 {
 	if (Property)
@@ -217,6 +229,11 @@ inline void InitializeValues(FProperty* Property, uint8* Items, int32 Num, SIZE_
 	}
 }
 
+inline void DestroyValue(FProperty* Property, void* Item)
+{
+	((*Property).*GDestroyPropertyValue)(Item);
+}
+
 inline void DestroyValues(FProperty* Property, uint8* Items, int32 Num, SIZE_T Stride)
 {
 	for (uint8* It = Items, *End = Items + NumBytes(Num, Stride); It != End; It += Stride)
@@ -468,6 +485,48 @@ static FRangeBinding AllocateArrayBinding(FArrayProperty* Property)
 
 ////////////////////////////////////////////////////////////////////////////////////////////////
 
+template<typename ScriptSet>
+inline bool IsCompact(const ScriptSet& Set)
+{
+	return Set.NumUnchecked() == Set.GetMaxIndex();
+}
+
+inline FExistingItemSlice GetContiguousSlice(int32 Idx, const FScriptSparseArray& Elems, const uint8* Data, SIZE_T Stride)
+{
+	int32 Num = 1;
+	for (;!Elems.IsValidIndex(Idx); ++Idx) {}
+	for (; Elems.IsValidIndex(Idx + Num); ++Num) {}
+	return { Data + NumBytes(Idx, Stride), static_cast<uint64>(Num) };
+}
+
+inline void ReadSparseItems(FExistingItems& Dst, const FScriptSparseArray& Src, const FScriptSparseArrayLayout& Layout)
+{
+	const uint8* Data = static_cast<const uint8*>(Src.GetData(0, Layout));
+
+	if (FExistingItemSlice LastRead = Dst.Slice)
+	{
+		// Continue partial response
+		int64 TotalBytesRead = static_cast<const uint8*>(LastRead.Data) - Data;
+		check(TotalBytesRead % Layout.Size == 0);
+		int32 NextIdx = /* skip known invalid */ 1 + TotalBytesRead / Layout.Size;
+		check(NextIdx < Src.GetMaxIndex());
+		check(NextIdx > static_cast<int32>(LastRead.Num));
+		Dst.Slice = GetContiguousSlice(NextIdx, Src, Data, Layout.Size);
+	}
+	else if (Src.IsCompact())
+	{
+		int32 Num = Src.Num();
+		Dst.SetAll(Num ? Data : nullptr, Num);
+	}
+	else
+	{
+		// Start partial response
+		Dst.NumTotal = Src.Num();
+		Dst.Stride = Layout.Size;
+		Dst.Slice = GetContiguousSlice(0, Src, Data, Layout.Size);
+	}
+}
+
 struct FSetBinding : IItemRangeBinding
 {
 	FProperty*					Element;
@@ -482,6 +541,11 @@ struct FSetBinding : IItemRangeBinding
 	, bConstruct(HasConstructor(In->ElementProp))
 	{}
 
+	virtual void ReadItems(FSaveRangeContext& Ctx) const override
+	{
+		ReadSparseItems(/* out */ Ctx.Items, Ctx.Request.GetRange<FScriptSparseArray>(), Layout.SparseArrayLayout);
+	}
+
 	virtual void MakeItems(FLoadRangeContext& Ctx) const override
 	{
 		FScriptSet& Set = Ctx.Request.GetRange<FScriptSet>();
@@ -491,8 +555,7 @@ struct FSetBinding : IItemRangeBinding
 		{
 			if (bDestruct && !Set.IsEmpty())
 			{
-				uint8* Data = static_cast<uint8*>(Set.GetData(0, Layout));
-				DestroyValues(Element, Data, Set.NumUnchecked(), Layout.Size);
+				DestructItems(Set);
 			}
 
 			Set.Empty(NewNum, Layout);
@@ -520,50 +583,255 @@ struct FSetBinding : IItemRangeBinding
 		}
 	}
 
+	inline void DestructItems(FScriptSet& Set) const
+	{
+		if (IsCompact(Set))
+		{
+			uint8* Data = static_cast<uint8*>(Set.GetData(0, Layout));
+			DestroyValues(Element, Data, Set.Num(), Layout.Size);
+		}
+		else
+		{
+			for (int32 Idx = 0, Num = Set.Num(); Idx < Num; ++Idx)
+			{
+				DestroyValue(Element, Set.GetData(Idx, Layout));
+			}
+		}
+	}
+};
+
+static constexpr ERangeSizeType DefaultSetSizeType = RangeSizeOf(FDefaultAllocator::SizeType{});
+
+static FRangeBinding AllocateSetBinding(FSetProperty* Property)
+{
+	// Todo: Ownership / memory leak
+	return FRangeBinding(*new FSetBinding(Property), DefaultSetSizeType);
+}
+
+////////////////////////////////////////////////////////////////////////////////////////////////
+
+enum class EMapFlags : uint32 { ConstructKey = 1, ConstructValue = 2, DestructKey = 4, DestructValue = 8 };
+
+// Currently can't extract constructor/destructor function pointers from FProperty, which
+// requires keeping FProperty* and prevents range binding reuse, @see AllocateMapBinding()
+template<class ScriptMap, EMapFlags Flags>
+struct TMapPropertyBinding : IItemRangeBinding
+{
+	static constexpr bool bConstructKey =		(uint32(Flags) & uint32(EMapFlags::ConstructKey)) != 0;
+	static constexpr bool bConstructValue =		(uint32(Flags) & uint32(EMapFlags::ConstructValue)) != 0;
+	static constexpr bool bDestructKey =		(uint32(Flags) & uint32(EMapFlags::DestructKey)) != 0;
+	static constexpr bool bDestructValue =		(uint32(Flags) & uint32(EMapFlags::DestructValue)) != 0;
+
+	const FScriptMapLayout Layout;
+	FProperty* KeyProp;
+	FProperty* ValueProp;
+
+	TMapPropertyBinding(FMapProperty* Property)
+	: Layout(Property->MapLayout)
+	, KeyProp(Property->KeyProp)
+	, ValueProp(Property->ValueProp)
+	{}
+
 	virtual void ReadItems(FSaveRangeContext& Ctx) const override
 	{
-		const FScriptSparseArray& Elems = Ctx.Request.GetRange<FScriptSparseArray>();
-		const uint8* Data = static_cast<const uint8*>(Elems.GetData(0, Layout.SparseArrayLayout));
+		ReadSparseItems(/* out */ Ctx.Items, Ctx.Request.GetRange<FScriptSparseArray>(), Layout.SetLayout.SparseArrayLayout);
+	}
+
+	virtual void MakeItems(FLoadRangeContext& Ctx) const override
+	{
+		ScriptMap& Map = Ctx.Request.GetRange<ScriptMap>();
 
-		if (FExistingItemSlice LastRead = Ctx.Items.Slice)
+		if (Ctx.Request.IsFirstCall())
 		{
-			// Continue partial response
-			int64 TotalBytesRead = static_cast<const uint8*>(LastRead.Data) - Data;
-			check(TotalBytesRead % Layout.Size == 0);
-			int32 NextIdx = /* skip known invalid */ 1 + TotalBytesRead / Layout.Size;
-			check(NextIdx < Elems.GetMaxIndex());
-			check(NextIdx > static_cast<int32>(LastRead.Num));
-			Ctx.Items.Slice = GetContiguousSlice(NextIdx, Elems, Data);
+			DestructItems(Map);
+
+			int32 NewNum = static_cast<int32>(Ctx.Request.NumTotal());
+			Map.Empty(NewNum, Layout);
+			if (NewNum)
+			{
+				void* Data = AddAndInitItems(Map, NewNum);
+				Ctx.Items.Set(Data, Ctx.Request.NumTotal(), Layout.SetLayout.Size);
+				Ctx.Items.RequestFinalCall();
+			}
 		}
-		else if (Elems.IsCompact())
+		else
 		{
-			int32 Num = Elems.Num();
-			Ctx.Items.SetAll(Num ? Data : nullptr, Num);
+			check(Ctx.Request.IsFinalCall());
+			Map.Rehash(Layout, [this](const void* Src) { return KeyProp->GetValueTypeHash(Src); });		
 		}
-		else
+	}
+
+	void* AddAndInitItems(ScriptMap& Map, const int32 Num) const
+	{
+		for (int32 Idx = 0; Idx < Num; ++Idx)
 		{
-			// Start partial response
-			Ctx.Items.NumTotal = Elems.Num();
-			Ctx.Items.Stride = Layout.Size;
-			Ctx.Items.Slice = GetContiguousSlice(0, Elems, Data);
+			Map.AddUninitialized(Layout);
 		}
+		check(IsCompact(Map));
+
+		uint8* Data = static_cast<uint8*>(Map.GetData(0, Layout));
+		SIZE_T Stride = static_cast<SIZE_T>(Layout.SetLayout.Size);
+		if constexpr (bConstructKey && bConstructValue == false)
+		{
+			FMemory::Memzero(Data, Stride * Num);
+		}
+		if constexpr (bConstructKey)
+		{
+			InitializeValues(KeyProp, Data, Num, Stride);
+		}
+		if constexpr (bConstructValue)
+		{
+			InitializeValues(ValueProp, Data + Layout.ValueOffset, Num, Stride);
+		}
+
+		return Data;
 	}
 
-	FExistingItemSlice GetContiguousSlice(int32 Idx, const FScriptSparseArray& Elems, const uint8* Data) const
+	inline void DestructItems(ScriptMap& Map) const
 	{
-		int32 Num = 1;
-		for (;!Elems.IsValidIndex(Idx); ++Idx) {}
-		for (; Elems.IsValidIndex(Idx + Num); ++Num) {}
-		return { Data + NumBytes(Idx, Layout.Size), static_cast<uint64>(Num) };
+		if constexpr (bDestructKey || bDestructValue)
+		{
+			if (const int32 Num = Map.Num())
+			{
+				if (IsCompact(Map))
+				{
+					uint8* Data = static_cast<uint8*>(Map.GetData(0, Layout));
+					if constexpr (bDestructKey)
+					{
+						DestroyValues(KeyProp, Data, Num, Layout.SetLayout.Size);
+					}
+					if constexpr (bDestructValue)
+					{
+						DestroyValues(ValueProp, Data + Layout.ValueOffset, Num, Layout.SetLayout.Size);
+					}
+				}
+				else
+				{
+					for (int32 Idx = 0; Idx < Num; ++Idx)
+					{
+						uint8* Item = static_cast<uint8*>(Map.GetData(Idx, Layout));
+						if constexpr (bDestructKey)
+						{
+							DestroyValue(KeyProp, Item);
+						}
+						if constexpr (bDestructValue)
+						{
+							DestroyValue(ValueProp, Item + Layout.ValueOffset);
+						}
+					}
+				}
+			}
+		}
 	}
 };
 
-static FRangeBinding AllocateSetBinding(FSetProperty* Property)
+inline uint32 GetMapPropertyFlags(FProperty* Key, FProperty* Val)
+{
+	return	HasConstructor(Key) * static_cast<uint32>(EMapFlags::ConstructKey) |
+			HasConstructor(Val) * static_cast<uint32>(EMapFlags::ConstructValue) |
+			HasDestructor(Key) *  static_cast<uint32>(EMapFlags::DestructKey) |
+			HasDestructor(Val) *  static_cast<uint32>(EMapFlags::DestructValue);
+}
+
+template<class ScriptMap>
+inline IItemRangeBinding* AllocateBinding(FMapProperty* Property)
+{
+	switch (GetMapPropertyFlags(Property->KeyProp, Property->ValueProp))
+	{
+	case 0:		return new TMapPropertyBinding<ScriptMap, EMapFlags(0)>(Property);
+	case 1:		return new TMapPropertyBinding<ScriptMap, EMapFlags(1)>(Property);
+	case 2:		return new TMapPropertyBinding<ScriptMap, EMapFlags(2)>(Property);
+	case 3:		return new TMapPropertyBinding<ScriptMap, EMapFlags(3)>(Property);
+	case 4:		return new TMapPropertyBinding<ScriptMap, EMapFlags(4)>(Property);
+	case 5:		return new TMapPropertyBinding<ScriptMap, EMapFlags(5)>(Property);
+	case 6:		return new TMapPropertyBinding<ScriptMap, EMapFlags(6)>(Property);
+	case 7:		return new TMapPropertyBinding<ScriptMap, EMapFlags(7)>(Property);
+	case 8:		return new TMapPropertyBinding<ScriptMap, EMapFlags(8)>(Property);
+	case 9:		return new TMapPropertyBinding<ScriptMap, EMapFlags(9)>(Property);
+	case 10:	return new TMapPropertyBinding<ScriptMap, EMapFlags(10)>(Property);
+	case 11:	return new TMapPropertyBinding<ScriptMap, EMapFlags(11)>(Property);
+	case 12:	return new TMapPropertyBinding<ScriptMap, EMapFlags(12)>(Property);
+	case 13:	return new TMapPropertyBinding<ScriptMap, EMapFlags(13)>(Property);
+	case 14:	return new TMapPropertyBinding<ScriptMap, EMapFlags(14)>(Property);
+	case 15:	return new TMapPropertyBinding<ScriptMap, EMapFlags(15)>(Property);
+	default:	check(false);
+	}
+	return nullptr;
+}
+
+static FRangeBinding AllocateMapBinding(FMapProperty* Property)
 {
+	bool bFreezable = EnumHasAnyFlags(Property->MapFlags, EMapPropertyFlags::UsesMemoryImageAllocator);
 	// Todo: Ownership / memory leak
-	return FRangeBinding(*new FSetBinding(Property), ERangeSizeType::S32);
+	IItemRangeBinding* Leak = bFreezable ? AllocateBinding<FFreezableScriptMap>(Property) : AllocateBinding<FScriptMap>(Property);
+	return FRangeBinding(*Leak, DefaultSetSizeType);
 }
 
+
+////////////////////////////////////////////////////////////////////////////////////////////////
+
+inline uint32 HashRangeBindings(TConstArrayView<FRangeBinding> In)
+{
+	return static_cast<uint32>(FXxHash64::HashBuffer(In.GetData(), In.NumBytes()).Hash);
+}
+
+inline uint32 HashSkipOffset(FMemberBinding In)
+{
+	uint32 Out = HashCombineFast(GetTypeHash(In.InnermostSchema), GetTypeHash(In.InnermostType));
+	return In.RangeBindings.IsEmpty() ? Out : HashCombineFast(Out, HashRangeBindings(In.RangeBindings));
+}
+
+inline bool EqSkipOffset(FMemberBinding A, FMemberBinding B)
+{
+	return A.InnermostType == B.InnermostType && A.InnermostSchema == B.InnermostSchema && Algo::Compare(A.RangeBindings, B.RangeBindings);
+}
+
+class FPairBindings
+{
+	struct FPair
+	{
+		FMemberBinding KV[2];
+		friend uint32 GetTypeHash(FPair In) { return HashCombineFast(HashSkipOffset(In.KV[0]), HashSkipOffset(In.KV[1])); };
+		inline bool operator==(const FPair& O) const { return EqSkipOffset(KV[0], O.KV[0]) && EqSkipOffset(KV[1], O.KV[1]); }
+	};
+
+	TMap<FPair, FStructSchemaId>	Bindings;
+
+	inline FStructSchemaId BindImpl(FPair Pair)
+	{
+		check(Pair.KV[0].Offset == 0 && Pair.KV[1].Offset > 0);
+		if (const FStructSchemaId* Id = Bindings.Find(Pair))
+		{
+			return *Id;
+		}
+
+		// Index names
+		FTypeId Params[2] = { Pair.KV[0].IndexParameterName(GUE.Names), Pair.KV[1].IndexParameterName(GUE.Names) };
+		FTypeId Type = { GUE.CommonIds.Core, FTypenameId(GUE.Names.MakeParametricTypeId(GUE.CommonIds.Pair, Params)) };
+		FStructSchemaId Id = GUE.Names.IndexStruct(Type);
+		FMemberId Members[2] = { GUE.CommonIds.Key, GUE.CommonIds.Value };
+
+		// Todo: Ownership / memory leak
+		GUE.Types.DeclareStruct(Id, Type, Members, EMemberPresence::RequireAll);
+		GUE.Schemas.BindStruct(Id, Id, Pair.KV);
+
+		Bindings.Emplace(Pair, Id);
+
+		return Id;
+	}
+public:
+
+	inline FMemberBinding Bind(FMemberBinding Key, FMemberBinding Value)
+	{
+		FMemberBinding Out(0);
+		Out.InnermostSchema = BindImpl(FPair{Key, Value});
+		Out.InnermostType = DefaultStructBindType;
+		return Key;
+	}
+};
+
+static FPairBindings GPairs;
+
 ////////////////////////////////////////////////////////////////////////////////////////////////
 
 struct FStaticArrayBinding : IItemRangeBinding
@@ -733,9 +1001,13 @@ private:
 	{
 		return BindAsRange(Property, AllocateArrayBinding(Property), BindSingleProperty(Property->Inner));
 	}
-	inline FMemberBinding BindMap(FMapProperty* Property) { return Todo(Property); }
+	inline FMemberBinding BindMap(FMapProperty* Property)
+	{
+		FMemberBinding PairBinding = GPairs.Bind(BindSingleProperty(Property->KeyProp), BindSingleProperty(Property->ValueProp));
+		return BindAsRange(Property, AllocateMapBinding(Property), PairBinding);
+	}
 	inline FMemberBinding BindSet(FSetProperty* Property)
-	{ 
+	{
 		return BindAsRange(Property, AllocateSetBinding(Property), BindSingleProperty(Property->ElementProp));
 	}
 	inline FMemberBinding BindOptional(FOptionalProperty* Property) { return Todo(Property); }
@@ -813,7 +1085,6 @@ private:
 		return FMemberBinding(Property->GetOffset_ForInternal());
 	}
 
-
 	FTypeId MakeStaticArrayTypename(FName PropertyName)
 	{
 		return { GetOwnerScope(), GUE.Names.MakeTypename(PropertyName) };
diff --git a/Engine/Plugins/Experimental/PlainPropsUObject/Source/Public/PlainPropsUObjectRuntime.h b/Engine/Plugins/Experimental/PlainPropsUObject/Source/Public/PlainPropsUObjectRuntime.h
--- a/Engine/Plugins/Experimental/PlainPropsUObject/Source/Public/PlainPropsUObjectRuntime.h
+++ b/Engine/Plugins/Experimental/PlainPropsUObject/Source/Public/PlainPropsUObjectRuntime.h
@@ -19,29 +19,35 @@ namespace PlainProps::UE
 struct FCommonUObjectIds
 {
 	explicit FCommonUObjectIds(TIdIndexer<FName>& Names);
+	
+	FScopeId					Core;
+	FScopeId					CoreUObject;
 
-	FScopeId			CoreUObject;
-	FStructSchemaId		Name;
-	FStructSchemaId		Text;
-	FStructSchemaId		FieldPath;
-	FStructSchemaId		SoftObjectPath;
-	FStructSchemaId		ObjectPtr;
-	FStructSchemaId		WeakObjectPtr;
-	FStructSchemaId		LazyObjectPtr;
-	FStructSchemaId		SoftObjectPtr;
-	FStructSchemaId		ScriptInterface;
+	FStructSchemaId				Name;
+	FStructSchemaId				Text;
+	FStructSchemaId				FieldPath;
+	FStructSchemaId				SoftObjectPath;
+	FStructSchemaId				ObjectPtr;
+	FStructSchemaId				WeakObjectPtr;
+	FStructSchemaId				LazyObjectPtr;
+	FStructSchemaId				SoftObjectPtr;
+	FStructSchemaId				ScriptInterface;
+
+	FConcreteTypenameId			Pair;
+	FMemberId					Key;
+	FMemberId					Value;
 };
 
 struct FGlobals
 {
 	FGlobals();
 
-	TIdIndexer<FName>	Names;
-	FDeclarations		Types;
-	FSchemaBindings		Schemas;
-	FCustomBindings		Customs;
-	FCustomBindings		DeltaCustoms;
-	FCommonUObjectIds	CommonIds;
+	TIdIndexer<FName>			Names;
+	FDeclarations				Types;
+	FSchemaBindings				Schemas;
+	FCustomBindings				Customs;
+	FCustomBindings				DeltaCustoms;
+	FCommonUObjectIds			CommonIds;
 };
 PLAINPROPSUOBJECT_API extern FGlobals GUE;
 
@@ -49,8 +55,8 @@ struct FIds
 {
 	static FNameId				IndexName(FAnsiStringView Name)			{ return GUE.Names.MakeName(FName(Name)); }
 	static FMemberId			IndexMember(FAnsiStringView Name)		{ return GUE.Names.NameMember(FName(Name)); }
-	static FTypenameId			IndexTypename(FAnsiStringView Name)		{ return GUE.Names.MakeTypename(FName(Name)); }
-	static FFlatScopeId			IndexScope(FAnsiStringView Name)		{ return GUE.Names.MakeScope(FName(Name)); }
+	static FConcreteTypenameId	IndexTypename(FAnsiStringView Name)		{ return GUE.Names.NameType(FName(Name)); }
+	static FFlatScopeId			IndexScope(FAnsiStringView Name)		{ return GUE.Names.NameScope(FName(Name)); }
 	static FEnumSchemaId		IndexEnum(FTypeId Type)					{ return GUE.Names.IndexEnum(Type); }
 	static FStructSchemaId		IndexStruct(FTypeId Type)				{ return GUE.Names.IndexStruct(Type); }
 	static FIdIndexerBase&		GetIndexer()							{ return GUE.Names; }
