PlainProps F[Ansi|Utf8]String bindings

#rnx

diff --git a/Engine/Plugins/Experimental/PlainProps/Source/Public/PlainPropsUeCoreBindings.h b/Engine/Plugins/Experimental/PlainProps/Source/Public/PlainPropsUeCoreBindings.h
--- a/Engine/Plugins/Experimental/PlainProps/Source/Public/PlainPropsUeCoreBindings.h
+++ b/Engine/Plugins/Experimental/PlainProps/Source/Public/PlainPropsUeCoreBindings.h
@@ -146,19 +146,21 @@ struct TArrayBinding : public IItemRangeBinding
 
 //////////////////////////////////////////////////////////////////////////
 
-struct FStringBinding : public ILeafRangeBinding
+template<class StringType>
+struct TStringBinding : ILeafRangeBinding
 {
 	using SizeType = int32;
 	using ItemType = char8_t;
+	using CharType = typename StringType::ElementType;
 
 	virtual void SaveLeaves(const void* Range, FLeafRangeAllocator& Out) const override
 	{
-		const TArray<TCHAR>& Src = static_cast<const FString*>(Range)->GetCharArray();
+		const TArray<CharType>& Src = static_cast<const StringType*>(Range)->GetCharArray();
 		int32 SrcLen = Src.Num() - 1;
 		if (SrcLen <= 0)
 		{
 		}
-		else if constexpr (sizeof(TCHAR) == sizeof(char8_t))
+		else if constexpr (sizeof(CharType) == sizeof(char8_t))
 		{
 			char8_t* Utf8 = Out.AllocateRange<char8_t>(SrcLen);
 			FMemory::Memcpy(Utf8, Src.GetData(), SrcLen);
@@ -174,7 +176,7 @@ struct FStringBinding : public ILeafRangeBinding
 
 	virtual void LoadLeaves(void* Range, FLeafRangeLoadView Items) const override
 	{
-		TArray<TCHAR>& Dst = static_cast<FString*>(Range)->GetCharArray();
+		TArray<CharType>& Dst = static_cast<StringType*>(Range)->GetCharArray();
 		TRangeView<char8_t> Utf8 = Items.As<char8_t>();
 		const UTF8CHAR* Src = reinterpret_cast<const UTF8CHAR*>(Utf8.begin());
 		int32 SrcLen = static_cast<int32>(Utf8.Num());
@@ -182,17 +184,17 @@ struct FStringBinding : public ILeafRangeBinding
 		{
 			Dst.Reset();
 		}
-		else if constexpr (sizeof(TCHAR) == sizeof(char8_t))
+		else if constexpr (sizeof(CharType) == sizeof(char8_t))
 		{
 			Dst.SetNum(SrcLen + 1);
 			FMemory::Memcpy(Dst.GetData(), Src, SrcLen);
-			Dst[SrcLen] = '\0';	
+			Dst[SrcLen] = CharType('\0');	
 		}
 		else
 		{
-			int32 DstLen = FPlatformString::ConvertedLength<TCHAR>(Src, SrcLen);
+			int32 DstLen = FPlatformString::ConvertedLength<CharType>(Src, SrcLen);
 			Dst.SetNum(DstLen + 1);
-			TCHAR* DstEnd = FPlatformString::Convert(Dst.GetData(), DstLen, Src, SrcLen);
+			CharType* DstEnd = FPlatformString::Convert(Dst.GetData(), DstLen, Src, SrcLen);
 			check(DstEnd - Dst.GetData() == DstLen);
 			*DstEnd = '\0';
 		}
@@ -200,9 +202,9 @@ struct FStringBinding : public ILeafRangeBinding
 
 	virtual bool DiffLeaves(const void* RangeA, const void* RangeB) const override
 	{
-		const FString& A = *static_cast<const FString*>(RangeA);
-		const FString& B = *static_cast<const FString*>(RangeB);
-		// Case-sensitive comparison
+		const StringType& A = *static_cast<const StringType*>(RangeA);
+		const StringType& B = *static_cast<const StringType*>(RangeB);
+		// Case-sensitive unnormalized comparison
 		return Diff(A.Len(), B.Len(), A.GetCharArray().GetData(), B.GetCharArray().GetData(), sizeof(TCHAR));
 	}
 };
@@ -957,7 +959,19 @@ struct TRangeBind<TArray<T, Allocator>>
 template<>
 struct TRangeBind<FString>
 {
-	using Type = UE::FStringBinding;
+	using Type = UE::TStringBinding<FString>;
+};
+
+template<>
+struct TRangeBind<FAnsiString>
+{
+	using Type = UE::TStringBinding<FAnsiString>;
+};
+
+template<>
+struct TRangeBind<FUtf8String>
+{
+	using Type = UE::TStringBinding<FUtf8String>;
 };
 
 template<typename T>
diff --git a/Engine/Plugins/Experimental/PlainPropsUObject/Source/Private/PlainPropsUObjectBindings.cpp b/Engine/Plugins/Experimental/PlainPropsUObject/Source/Private/PlainPropsUObjectBindings.cpp
--- a/Engine/Plugins/Experimental/PlainPropsUObject/Source/Private/PlainPropsUObjectBindings.cpp
+++ b/Engine/Plugins/Experimental/PlainPropsUObject/Source/Private/PlainPropsUObjectBindings.cpp
@@ -1,6 +1,7 @@
 // Copyright Epic Games, Inc. All Rights Reserved.
 
 #include "PlainPropsUObjectRuntime.h"
+#include "PlainPropsUeCoreBindings.h"
 #include "Algo/Compare.h"
 #include "Containers/PagedArray.h"
 #include "Hash/xxhash.h"
@@ -244,6 +245,8 @@ inline void DestroyValues(FProperty* Property, uint8* Items, int32 Num, SIZE_T S
 
 ////////////////////////////////////////////////////////////////////////////////////////////////
 
+static constexpr ERangeSizeType DefaultSizeType = RangeSizeOf(FDefaultAllocator::SizeType{});
+
 // Helps cache array property range bindings
 struct FArrayPropertyInfo
 {
@@ -393,7 +396,7 @@ struct TLeafArrayBinding : ILeafRangeBinding
 // Reusable cache of FArrayProperty range bindings
 class FArrayPropertyBindings
 {
-	static constexpr ERangeSizeType SizeType = RangeSizeOf(FDefaultAllocator::SizeType{});
+	static constexpr ERangeSizeType SizeType = DefaultSizeType;
 	
 	const TLeafArrayBinding<ELeafType::Bool, ELeafWidth::B8>	Bool;
 	const TLeafArrayBinding<ELeafType::Float, ELeafWidth::B32>	Float;
@@ -600,12 +603,10 @@ struct FSetBinding : IItemRangeBinding
 	}
 };
 
-static constexpr ERangeSizeType DefaultSetSizeType = RangeSizeOf(FDefaultAllocator::SizeType{});
-
 static FRangeBinding AllocateSetBinding(FSetProperty* Property)
 {
 	// Todo: Ownership / memory leak
-	return FRangeBinding(*new FSetBinding(Property), DefaultSetSizeType);
+	return FRangeBinding(*new FSetBinding(Property), DefaultSizeType);
 }
 
 ////////////////////////////////////////////////////////////////////////////////////////////////
@@ -764,10 +765,9 @@ static FRangeBinding AllocateMapBinding(FMapProperty* Property)
 	bool bFreezable = EnumHasAnyFlags(Property->MapFlags, EMapPropertyFlags::UsesMemoryImageAllocator);
 	// Todo: Ownership / memory leak
 	IItemRangeBinding* Leak = bFreezable ? AllocateBinding<FFreezableScriptMap>(Property) : AllocateBinding<FScriptMap>(Property);
-	return FRangeBinding(*Leak, DefaultSetSizeType);
+	return FRangeBinding(*Leak, DefaultSizeType);
 }
 
-
 ////////////////////////////////////////////////////////////////////////////////////////////////
 
 inline uint32 HashRangeBindings(TConstArrayView<FRangeBinding> In)
@@ -832,6 +832,32 @@ public:
 
 static FPairBindings GPairs;
 
+//////////////////////////////////////////////////////////////////////////////////////////////
+
+struct FStringBindings
+{
+	TStringBinding<FString>		TChrInstance;
+	TStringBinding<FAnsiString>	AnsiInstance;
+	TStringBinding<FUtf8String>	Utf8Instance;
+	FRangeBinding				TChr{TChrInstance, DefaultSizeType};
+	FRangeBinding				Ansi{AnsiInstance, DefaultSizeType};
+	FRangeBinding				Utf8{Utf8Instance, DefaultSizeType};
+
+	FMemberBinding				Bind(FProperty* Property, uint64 CastFlags) const
+	{
+		check(FMath::CountBits(CastFlags & StringMask) == 1);
+		bool bTChr = HasAny<CASTCLASS_FStrProperty>(CastFlags);
+		bool bAnsi = HasAny<CASTCLASS_FAnsiStrProperty>(CastFlags);
+		const FRangeBinding& Binding = bTChr ? TChr : (bAnsi ? Ansi : Utf8);
+
+		FMemberBinding Out(Property->GetOffset_ForInternal());
+		Out.InnermostType = FMemberBindType(ReflectLeaf<char8_t>);
+		Out.RangeBindings = MakeArrayView(&Binding, 1);
+		return Out;
+	}
+};
+static const FStringBindings GStrings;
+
 ////////////////////////////////////////////////////////////////////////////////////////////////
 
 struct FStaticArrayBinding : IItemRangeBinding
@@ -876,13 +902,14 @@ private:
 			return MakeArrayView(&Ranges.Add_GetRef(Head), 1);
 		}
 		
-		// Ensure contiguous range by skipping last items in page
+		// Ensure contiguous out range by padding up with dummy tail slice
 		static constexpr uint32 PageMax = decltype(Ranges)::MaxPerPage();
 		const int32 OutNum = 1 + Tail.Num();
 		check(OutNum <= PageMax);
-		if (Ranges.NumPages() != (Ranges.Num() + Tail.Num()) % PageMax)
+		const int32 NewPages = Align(Ranges.Num() + OutNum, PageMax) / PageMax;
+		if (NewPages > Ranges.NumPages() && !Ranges.IsEmpty())
 		{
-			int32 NumPad = Ranges.MaxPerPage() - Ranges.Num() % PageMax;
+			int32 NumPad = PageMax - Ranges.Num() % PageMax;
 			Ranges.Append(Tail.Slice(0, NumPad));
 			check(Ranges.Num() % PageMax == 0);
 		}
@@ -1012,8 +1039,6 @@ private:
 	}
 	inline FMemberBinding BindOptional(FOptionalProperty* Property) { return Todo(Property); }
 
-	template<typename StrProperty>
-	inline FMemberBinding BindStr(StrProperty* Property) { return Todo(Property); }
 
 	#if WITH_VERSE_VM
 	inline FMemberBinding BindVValue(FVValueProperty* Property) { return Todo(Property); }
@@ -1063,14 +1088,7 @@ private:
 		}
 		else if (HasAny<StringMask>(Flags))
 		{
-			if (HasAny<CASTCLASS_FStrProperty>(Flags))
-			{
-				return BindStr(static_cast<FStrProperty*>(Property));
-			}
-		
-			return HasAny<CASTCLASS_FAnsiStrProperty>(Flags)
-				? BindStr(static_cast<FAnsiStrProperty*>(Property))
-				: BindStr(static_cast<FUtf8StrProperty*>(Property));
+			return GStrings.Bind(Property, Flags);
 		}
 	#if WITH_VERSE_VM
 		else if (HasAny<CASTCLASS_FVValueProperty | CASTCLASS_FVRestValueProperty>(Flags))
diff --git a/Engine/Plugins/Experimental/PlainPropsUObject/Source/Public/PlainPropsUObjectRuntime.h b/Engine/Plugins/Experimental/PlainPropsUObject/Source/Public/PlainPropsUObjectRuntime.h
--- a/Engine/Plugins/Experimental/PlainPropsUObject/Source/Public/PlainPropsUObjectRuntime.h
+++ b/Engine/Plugins/Experimental/PlainPropsUObject/Source/Public/PlainPropsUObjectRuntime.h
@@ -10,7 +10,8 @@
 
 namespace PlainProps
 {
-template<> inline void AppendString(FString& Out, const FName& Name) { Name.AppendString(Out); }
+template<>
+PLAINPROPS_API void AppendString(FString& Out, const FName& Name);
 }
 
 namespace PlainProps::UE
