PlainProps FOptionalProperty binding

#rnx

diff --git a/Engine/Plugins/Experimental/PlainPropsUObject/Source/Private/PlainPropsUObjectBindings.cpp b/Engine/Plugins/Experimental/PlainPropsUObject/Source/Private/PlainPropsUObjectBindings.cpp
--- a/Engine/Plugins/Experimental/PlainPropsUObject/Source/Private/PlainPropsUObjectBindings.cpp
+++ b/Engine/Plugins/Experimental/PlainPropsUObject/Source/Private/PlainPropsUObjectBindings.cpp
@@ -35,6 +35,7 @@ FCommonUObjectIds::FCommonUObjectIds(TIdIndexer<FName>& Names)
 , LazyObjectPtr(	Names.IndexStruct({CoreUObject, Names.MakeTypename("LazyObjectPtr")}))
 , SoftObjectPtr(	Names.IndexStruct({CoreUObject, Names.MakeTypename("SoftObjectPtr")}))
 , ScriptInterface(	Names.IndexStruct({CoreUObject, Names.MakeTypename("ScriptInterface")}))
+, Optional(			Names.NameType("Optional"))
 , Pair(				Names.NameType("Pair"))
 , Key(				Names.NameMember("Key"))
 , Value(			Names.NameMember("Value"))
@@ -826,7 +827,7 @@ public:
 		FMemberBinding Out(0);
 		Out.InnermostSchema = BindImpl(FPair{Key, Value});
 		Out.InnermostType = DefaultStructBindType;
-		return Key;
+		return Out;
 	}
 };
 
@@ -834,6 +835,53 @@ static FPairBindings GPairs;
 
 //////////////////////////////////////////////////////////////////////////////////////////////
 
+class FOptionalBindings
+{
+	struct FValueBinding : FMemberBinding
+	{
+		friend uint32 GetTypeHash(FValueBinding In) { return HashSkipOffset(In); };
+		inline bool operator==(FValueBinding O) const { return EqSkipOffset(*this, O); }
+	};
+
+	TMap<FValueBinding, FStructSchemaId>	Bindings;
+
+	inline FStructSchemaId BindImpl(const FValueBinding Value)
+	{
+		check(Value.Offset == 0);
+		if (const FStructSchemaId* Id = Bindings.Find(Value))
+		{
+			return *Id;
+		}
+
+		// Index names
+		FTypeId Params[1] = { Value.IndexParameterName(GUE.Names) };
+		FTypeId Type = { GUE.CommonIds.Core, FTypenameId(GUE.Names.MakeParametricTypeId(GUE.CommonIds.Optional, Params)) };
+		FStructSchemaId Id = GUE.Names.IndexStruct(Type);
+		FMemberId Members[1] = { GUE.CommonIds.Value };
+		FMemberBinding Inners[1] = { Value };
+
+		// Todo: Ownership / memory leak
+		GUE.Types.DeclareStruct(Id, Type, Members, EMemberPresence::AllowSparse);
+		GUE.Schemas.BindStruct(Id, Id, Inners);
+
+		Bindings.Emplace(Value, Id);
+
+		return Id;
+	}
+public:
+
+	inline FMemberBinding Bind(FMemberBinding Value)
+	{
+		FMemberBinding Out(0);
+		Out.InnermostSchema = BindImpl(static_cast<FValueBinding>(Value));
+		Out.InnermostType = DefaultStructBindType;
+		return Out;
+	}
+};
+
+static FOptionalBindings GOptionals;
+
+//////////////////////////////////////////////////////////////////////////////////////////////
 struct FStringBindings
 {
 	TStringBinding<FString>		TChrInstance;
@@ -1037,7 +1085,10 @@ private:
 	{
 		return BindAsRange(Property, AllocateSetBinding(Property), BindSingleProperty(Property->ElementProp));
 	}
-	inline FMemberBinding BindOptional(FOptionalProperty* Property) { return Todo(Property); }
+	inline FMemberBinding BindOptional(FOptionalProperty* Property)
+	{
+		return GOptionals.Bind(BindSingleProperty(Property->GetValueProperty()));
+	}
 
 
 	#if WITH_VERSE_VM
diff --git a/Engine/Plugins/Experimental/PlainPropsUObject/Source/Public/PlainPropsUObjectRuntime.h b/Engine/Plugins/Experimental/PlainPropsUObject/Source/Public/PlainPropsUObjectRuntime.h
--- a/Engine/Plugins/Experimental/PlainPropsUObject/Source/Public/PlainPropsUObjectRuntime.h
+++ b/Engine/Plugins/Experimental/PlainPropsUObject/Source/Public/PlainPropsUObjectRuntime.h
@@ -33,7 +33,8 @@ struct FCommonUObjectIds
 	FStructSchemaId				LazyObjectPtr;
 	FStructSchemaId				SoftObjectPtr;
 	FStructSchemaId				ScriptInterface;
-
+	
+	FConcreteTypenameId			Optional;
 	FConcreteTypenameId			Pair;
 	FMemberId					Key;
 	FMemberId					Value;
