PlainProps FSetProperty delta binding and FLeafRangeLoadView improvements

#rnx

diff --git a/Engine/Plugins/Experimental/PlainProps/Source/Private/PlainPropsBind.cpp b/Engine/Plugins/Experimental/PlainProps/Source/Private/PlainPropsBind.cpp
--- a/Engine/Plugins/Experimental/PlainProps/Source/Private/PlainPropsBind.cpp
+++ b/Engine/Plugins/Experimental/PlainProps/Source/Private/PlainPropsBind.cpp
@@ -139,6 +139,15 @@ FRangeBinding::FRangeBinding(const ILeafRangeBinding& Binding, ERangeSizeType Si
 
 ////////////////////////////////////////////////////////////////////////////////////////////////
 
+FRangeMemberBinding GetInnerRange(FRangeMemberBinding In)
+{
+	check(In.NumRanges > 1);
+	check(In.InnerTypes[0].IsRange());
+	return { In.InnerTypes + 1, In.RangeBindings + 1, static_cast<uint16>(In.NumRanges - 1), In.InnermostSchema };
+}
+
+////////////////////////////////////////////////////////////////////////////////////////////////
+
 void* FLeafRangeAllocator::Allocate(uint64 Num, SIZE_T LeafSize)
 {
 	check(!Range);
diff --git a/Engine/Plugins/Experimental/PlainProps/Source/Private/PlainPropsLoad.cpp b/Engine/Plugins/Experimental/PlainProps/Source/Private/PlainPropsLoad.cpp
--- a/Engine/Plugins/Experimental/PlainProps/Source/Private/PlainPropsLoad.cpp
+++ b/Engine/Plugins/Experimental/PlainProps/Source/Private/PlainPropsLoad.cpp
@@ -1055,14 +1055,6 @@ void FNestedRangeLoadIterator::operator++()
 
 ///////////////////////////////////////////////////////////////////////////
 
-FLeafRangeView FRangeLoadView::AsLeaves() const
-{
-	FUnpackedLeafType Leaf = Schema.ItemType.AsLeaf();
-	FOptionalEnumSchemaId EnumId = static_cast<FOptionalEnumSchemaId>(Schema.InnermostId);
-	const uint8* Data = static_cast<const uint8*>(Values.GetData());
-	return { Leaf, Schema.Batch.ReadId, EnumId, NumItems, Data };
-}
-
 FStructRangeLoadView FRangeLoadView::AsStructs() const
 {
 	check(IsStructRange());
diff --git a/Engine/Plugins/Experimental/PlainProps/Source/Private/PlainPropsSave.cpp b/Engine/Plugins/Experimental/PlainProps/Source/Private/PlainPropsSave.cpp
--- a/Engine/Plugins/Experimental/PlainProps/Source/Private/PlainPropsSave.cpp
+++ b/Engine/Plugins/Experimental/PlainProps/Source/Private/PlainPropsSave.cpp
@@ -189,13 +189,6 @@ template<class SaverType, typename InnerContextType>
 	return nullptr;
 }
 
-[[nodiscard]] static FRangeMemberBinding GetInnerRange(FRangeMemberBinding Member)
-{
-	check(Member.NumRanges > 1);
-	check(Member.InnerTypes[0].IsRange());
-	return { Member.InnerTypes + 1, Member.RangeBindings + 1, static_cast<uint16>(Member.NumRanges - 1), Member.InnermostSchema };
-}
-
 FBuiltRange* SaveLeafRange(const void* Range, const ILeafRangeBinding& Binding, FUnpackedLeafType Leaf, const FSaveContext& Ctx)
 {
 	FLeafRangeAllocator Allocator(Ctx.Scratch, Leaf);
@@ -257,7 +250,7 @@ static ELeafWidth GetBasicWidth(FLeafBindType Leaf)
 
 [[nodiscard]] static FBuiltStruct* SaveRangeItem(const uint8* Struct, FDefaultStruct Default, const FSaveContext& Ctx)
 {
-	return SaveStructDelta(Struct, Default.Struct, Default.Id, Ctx);
+	return SaveStructDeltaIfDiff(Struct, Default.Struct, Default.Id, Ctx);
 }
 
 [[nodiscard]] static FMemberType ToMemberType(FMemberBindType In)
@@ -527,7 +520,7 @@ static void SaveMemberDelta(FMemberBuilder& Out, const void* Struct, const void*
 
 static void SaveMemberDelta(FMemberBuilder& Out, const void* Struct, const void* Default, FMemberId Name, const FSaveContext& Ctx, FStructMemberBinding Member)
 {
-	if (FBuiltStruct* Delta = SaveStructDelta(At(Struct, Member.Offset), At(Default, Member.Offset), Member.Id, Ctx))
+	if (FBuiltStruct* Delta = SaveStructDeltaIfDiff(At(Struct, Member.Offset), At(Default, Member.Offset), Member.Id, Ctx))
 	{
 		Out.AddStruct(Name, Member.Id, MoveTemp(Delta));
 	}
@@ -539,9 +532,9 @@ FBuiltStruct* SaveStructDelta(const void* Struct, const void* Default, FStructSc
 	FMemberBuilder Out;
 	if (FCustomBindingEntry Custom = Ctx.Customs.FindStructToSave(BindId))
 	{
+		Declaration = &Ctx.Declarations.Get(Custom.DeclId);	
 		if (Custom.Binding->DiffCustom(Struct, Default))
 		{
-			Declaration = &Ctx.Declarations.Get(Custom.DeclId);	
 			Custom.Binding->SaveCustom(Out, Struct, Default, Ctx);
 		}
 	}
@@ -563,6 +556,38 @@ FBuiltStruct* SaveStructDelta(const void* Struct, const void* Default, FStructSc
 		}
 	}
 	
+	return Out.BuildAndReset(Ctx.Scratch, *Declaration, Ctx.Declarations.GetDebug());
+}
+
+FBuiltStruct* SaveStructDeltaIfDiff(const void* Struct, const void* Default, FStructSchemaId BindId, const FSaveContext& Ctx)
+{
+	FMemberBuilder Out;
+	if (FCustomBindingEntry Custom = Ctx.Customs.FindStructToSave(BindId))
+	{
+		if (Custom.Binding->DiffCustom(Struct, Default))
+		{
+			const FStructDeclaration* Declaration = &Ctx.Declarations.Get(Custom.DeclId);	
+			Custom.Binding->SaveCustom(Out, Struct, Default, Ctx);
+			return Out.BuildAndReset(Ctx.Scratch, *Declaration, Ctx.Declarations.GetDebug());
+		}
+		return nullptr;
+	}
+
+	const FSchemaBinding& Schema = Ctx.Schemas.GetStruct(BindId);
+	const FStructDeclaration* Declaration = &Ctx.Declarations.Get(Schema.DeclId);	
+	TConstArrayView<FMemberId> MemberOrder = Declaration->GetMemberOrder();
+
+	for (FMemberVisitor It(Schema); It.HasMore(); )
+	{
+		FMemberId Name = MemberOrder[It.GetIndex()];
+		switch (It.PeekKind())
+		{
+			case EMemberKind::Leaf:		SaveMemberDelta(Out, Struct, Default, Name, Ctx, It.GrabLeaf());	break;
+			case EMemberKind::Range:	SaveMemberDelta(Out, Struct, Default, Name, Ctx, It.GrabRange());	break;
+			case EMemberKind::Struct:	SaveMemberDelta(Out, Struct, Default, Name, Ctx, It.GrabStruct());	break;
+		}
+	}
+		
 	return Out.IsEmpty() ? nullptr : Out.BuildAndReset(Ctx.Scratch, *Declaration, Ctx.Declarations.GetDebug());
 }
 
diff --git a/Engine/Plugins/Experimental/PlainProps/Source/Public/PlainPropsBind.h b/Engine/Plugins/Experimental/PlainProps/Source/Public/PlainPropsBind.h
--- a/Engine/Plugins/Experimental/PlainProps/Source/Public/PlainPropsBind.h
+++ b/Engine/Plugins/Experimental/PlainProps/Source/Public/PlainPropsBind.h
@@ -205,6 +205,8 @@ struct FStructMemberBinding
 	SIZE_T					Offset;
 };
 
+PLAINPROPS_API FRangeMemberBinding GetInnerRange(FRangeMemberBinding In); // @pre In.NumRanges > 1
+
 ////////////////////////////////////////////////////////////////////////////////////////////////
 
 struct FDualStructSchemaId
@@ -309,6 +311,7 @@ struct FRequireAll
 
 ////////////////////////////////////////////////////////////////////////////////////////////////
 
+// Request from loading layer to IItemRangeBinding to allocate and construct items
 class FConstructionRequest
 {
 	friend FRangeLoader;
@@ -330,6 +333,10 @@ public:
 	bool IsFinalCall() const { return Index == Num;}
 };
 
+// Response from IItemRangeBinding with contiguous items ready to be loaded
+//
+// Non-contiguous containers provide one items one by one or use
+// FLoadRangeContext::Scratch or a temp allocation to avoid that
 class FConstructedItems
 {
 public:
@@ -352,6 +359,12 @@ public:
 		Num = NumItems;
 		Size = ItemSize;
 	}
+	
+	void UpdateNum(uint64 NumItems)
+	{
+		check(Data);
+		Num = NumItems;
+	}
 
 	template<typename ItemType>
 	TArrayView<ItemType> Get()
@@ -359,6 +372,9 @@ public:
 		return MakeArrayView(reinterpret_cast<ItemType*>(Data), Num);
 	}
 
+	inline uint8* GetData()		{ return Data; }
+	inline uint8* GetDataEnd()	{ return Data + NumBytes(); }
+
 private:
 	friend FRangeLoader;
 	uint8*	Data = nullptr;
@@ -469,65 +485,6 @@ public:
 	FBuiltRange* GetAllocatedRange() { return Range; }
 };
 
-// Hides internal representation to enable future format changes, 
-// e.g. store zeroes or 1.0f in some compact fashion or var int encodings
-class FLeafRangeLoadView
-{
-	const void*				Data;
-	uint64					NumItems;
-	FUnpackedLeafType		Leaf;
-
-	static constexpr bool CopyCompatible(ELeafType Dst, ELeafType Src)
-	{
-		return Dst == Src || (Dst == ELeafType::IntU && (Src == ELeafType::Hex || Src == ELeafType::Enum));
-	}
-
-public:
-	FLeafRangeLoadView(const void* InData, uint64 InNum, FUnpackedLeafType InLeaf)
-	: Data(InData)
-	, NumItems(InNum)
-	, Leaf(InLeaf)
-	{}
-
-	uint64 Num() const { return NumItems; }
-	
-	template<typename T>
-	auto As() const
-	{
-		check(Leaf == ReflectLeaf<T>);
-		return TRangeView<T>(static_cast<const T*>(Data), NumItems);
-	}
-
-	template<>
-	auto As<bool>() const
-	{
-		check(Leaf.Type == ELeafType::Bool && Leaf.Width == ELeafWidth::B8);
-		return FBoolRangeView(static_cast<const uint8*>(Data), NumItems);
-	}
-
-	template<ELeafType Type, ELeafWidth Width>
-	void Copy(void* Dst, uint64 NumBytes) const
-	{
-		check(Width == Leaf.Width && CopyCompatible(Type, Leaf.Type));
-		check(NumBytes == NumItems * SizeOf(Leaf.Width));
-		if (NumBytes)
-		{
-			FMemory::Memcpy(Dst, Data, NumBytes);
-		}
-	}
-
-	template<>
-	void Copy<ELeafType::Bool, ELeafWidth::B8>(void* Dst, uint64 NumBytes) const
-	{
-		check(Leaf == ReflectArithmetic<bool>);
-		check(NumBytes == NumItems);
-		bool* bOut = static_cast<bool*>(Dst);
-		for (bool b : As<bool>())
-		{
-			*bOut++ = b;
-		}	
-	}
-};
 
 // Specialized binding for transcoding leaf ranges
 struct alignas(16) ILeafRangeBinding
diff --git a/Engine/Plugins/Experimental/PlainProps/Source/Public/PlainPropsBuild.h b/Engine/Plugins/Experimental/PlainProps/Source/Public/PlainPropsBuild.h
--- a/Engine/Plugins/Experimental/PlainProps/Source/Public/PlainPropsBuild.h
+++ b/Engine/Plugins/Experimental/PlainProps/Source/Public/PlainPropsBuild.h
@@ -219,11 +219,11 @@ struct FTypedValue
 struct FBuiltMember
 {
 	FBuiltMember(FMemberId InName, FTypedValue In) : FBuiltMember(InName, In.Schema, In.Value) {}
-	FBuiltMember(FMemberId Name, FUnpackedLeafType Leaf, FOptionalEnumSchemaId Schema, uint64 Value);
-	FBuiltMember(FMemberId Name, FTypedRange Range);
-	FBuiltMember(FMemberId Name, FStructSchemaId Schema, FBuiltStruct* Value);
 	FBuiltMember(FOptionalMemberId N, FMemberSchema S, FBuiltValue V) : Name(N), Schema(MoveTemp(S)), Value(V) {}
-	static FBuiltMember MakeSuper(FStructSchemaId Schema, FBuiltStruct* Value);
+	PLAINPROPS_API FBuiltMember(FMemberId Name, FUnpackedLeafType Leaf, FOptionalEnumSchemaId Schema, uint64 Value);
+	PLAINPROPS_API FBuiltMember(FMemberId Name, FTypedRange Range);
+	PLAINPROPS_API FBuiltMember(FMemberId Name, FStructSchemaId Schema, FBuiltStruct* Value);
+	PLAINPROPS_API static FBuiltMember MakeSuper(FStructSchemaId Schema, FBuiltStruct* Value);
 
 	FOptionalMemberId		Name;
 	FMemberSchema			Schema;
diff --git a/Engine/Plugins/Experimental/PlainProps/Source/Public/PlainPropsLoadMember.h b/Engine/Plugins/Experimental/PlainProps/Source/Public/PlainPropsLoadMember.h
--- a/Engine/Plugins/Experimental/PlainProps/Source/Public/PlainPropsLoadMember.h
+++ b/Engine/Plugins/Experimental/PlainProps/Source/Public/PlainPropsLoadMember.h
@@ -3,6 +3,7 @@
 #pragma once
 
 #include "PlainPropsLoad.h"
+#include <type_traits>
 
 namespace PlainProps 
 {
@@ -52,7 +53,7 @@ public:
 	bool								IsStructRange() const	{ return Schema.ItemType.GetKind() == EMemberKind::Struct; }
 	bool								IsNestedRange() const	{ return Schema.ItemType.GetKind() == EMemberKind::Range; }
 
-	PLAINPROPS_API FLeafRangeView		AsLeaves() const;		// @pre IsLeafRange()
+	FLeafRangeLoadView					AsLeaves() const;		// @pre IsLeafRange()
 	PLAINPROPS_API FStructRangeLoadView	AsStructs() const;		// @pre IsStructRange()
 	PLAINPROPS_API FNestedRangeLoadView	AsRanges() const;		// @pre IsNestedRange()
 
@@ -68,6 +69,68 @@ PLAINPROPS_API void	LoadRange(void* Dst, FRangeLoadView Src, TConstArrayView<FRa
 
 //////////////////////////////////////////////////////////////////////////
 
+// Hides internal representation to enable future format changes, 
+// e.g. store zeroes or 1.0f in some compact fashion or var int encodings
+class FLeafRangeLoadView
+{
+	const void*				Data;
+	uint64					NumItems;
+	FUnpackedLeafType		Leaf;
+
+public:
+	FLeafRangeLoadView(const void* InData, uint64 InNum, FUnpackedLeafType InLeaf)
+	: Data(InData)
+	, NumItems(InNum)
+	, Leaf(InLeaf)
+	{}
+
+	uint64 Num() const { return NumItems; }
+	
+	template<typename T>
+	auto As() const
+	{
+		check(Leaf == ReflectLeaf<T>);
+		return TRangeView<T>(static_cast<const T*>(Data), NumItems);
+	}
+
+	template<>
+	auto As<bool>() const
+	{
+		check(Leaf.Type == ELeafType::Bool && Leaf.Width == ELeafWidth::B8);
+		return FBoolRangeView(static_cast<const uint8*>(Data), NumItems);
+	}
+
+	template<typename T>
+	auto AsBitCast() const requires (std::is_unsigned_v<T>)
+	{
+		if constexpr (std::is_same_v<T, bool>)
+		{
+			return As<bool>();
+		}
+		else
+		{
+			check(Leaf.Type != ELeafType::Bool);
+			check(SizeOf(Leaf.Width) == sizeof(T));
+			return TRangeView<T>(reinterpret_cast<const T*>(Data), NumItems);
+		}
+	}
+	
+	template<ELeafType Type, ELeafWidth Width>
+	auto AsBitCast() const
+	{
+		if constexpr (Type == ELeafType::Bool)			return As<bool>();
+		else if constexpr (Width == ELeafWidth::B8)		return AsBitCast<uint8>();
+		else if constexpr (Width == ELeafWidth::B16)	return AsBitCast<uint16>();
+		else if constexpr (Width == ELeafWidth::B32)	return AsBitCast<uint32>();
+		else if constexpr (Width == ELeafWidth::B64)	return AsBitCast<uint64>();
+	}
+};
+
+inline FLeafRangeLoadView FRangeLoadView::AsLeaves() const		{ return { Values.GetData(), NumItems, Schema.ItemType.AsLeaf() }; }
+inline FLeafRangeLoadView FLeafRangeView::AsLoadView() const	{ return { Values, NumItems, {Type, Width} }; }
+
+//////////////////////////////////////////////////////////////////////////
+
 class FNestedRangeLoadIterator
 {
 	friend FNestedRangeLoadView;
@@ -108,18 +171,18 @@ public:
 class FMemberLoader
 {
 public:
-	explicit FMemberLoader(FStructLoadView Struct);
+	PLAINPROPS_API explicit FMemberLoader(FStructLoadView Struct);
 
-	bool					HasMore() const				{ return Reader.HasMore(); }
+	bool							HasMore() const				{ return Reader.HasMore(); }
 	
-	FOptionalMemberId		PeekName() const			{ return Reader.PeekName(); }
-	FOptionalMemberId		PeekNameUnchecked() const	{ return Reader.PeekNameUnchecked(); }
-	EMemberKind				PeekKind() const			{ return Reader.PeekKind(); }
-	FMemberType				PeekType() const			{ return Reader.PeekType(); }
-
-	FLeafView				GrabLeaf()					{ return Reader.GrabLeaf(); }
-	FRangeLoadView			GrabRange();				// @pre PeekKind() == EMemberKind::Range
-	FStructLoadView			GrabStruct();				// @pre PeekKind() == EMemberKind::Struct
+	FOptionalMemberId				PeekName() const			{ return Reader.PeekName(); }
+	FOptionalMemberId				PeekNameUnchecked() const	{ return Reader.PeekNameUnchecked(); }
+	EMemberKind						PeekKind() const			{ return Reader.PeekKind(); }
+	FMemberType						PeekType() const			{ return Reader.PeekType(); }
+
+	FLeafView						GrabLeaf()					{ return Reader.GrabLeaf(); }
+	PLAINPROPS_API FRangeLoadView	GrabRange();				// @pre PeekKind() == EMemberKind::Range
+	PLAINPROPS_API FStructLoadView	GrabStruct();				// @pre PeekKind() == EMemberKind::Struct
 
 private:
 	FMemberReader			Reader;
diff --git a/Engine/Plugins/Experimental/PlainProps/Source/Public/PlainPropsRead.h b/Engine/Plugins/Experimental/PlainProps/Source/Public/PlainPropsRead.h
--- a/Engine/Plugins/Experimental/PlainProps/Source/Public/PlainPropsRead.h
+++ b/Engine/Plugins/Experimental/PlainProps/Source/Public/PlainPropsRead.h
@@ -12,6 +12,7 @@
 namespace PlainProps
 {
 
+class FLeafRangeLoadView;
 class FMemberLoader;
 class FMemberPrinter;
 struct FEnumSchema;
@@ -339,10 +340,11 @@ public:
 	, NumBits(InNum)
 	{}
 
-	uint64 Num() const					{ return NumBits;}
-	bool operator[](uint64 Idx) const	{ return *FBoolRangeIterator(Data, Idx); }
-	FBoolRangeIterator begin() const	{ return FBoolRangeIterator(Data, 0); }
-	FBoolRangeIterator end() const		{ return FBoolRangeIterator(Data, NumBits); }
+	inline uint64				Num() const								{ return NumBits;}
+	inline bool					operator[](uint64 Idx) const			{ return *FBoolRangeIterator(Data, Idx); }
+	inline void					Copy(void* Dst, uint64 NumBytes) const;
+	inline FBoolRangeIterator	begin() const							{ return FBoolRangeIterator(Data, 0); }
+	inline FBoolRangeIterator	end() const								{ return FBoolRangeIterator(Data, NumBits); }
 
 private:
 	const uint8* Data;
@@ -362,10 +364,11 @@ public:
 	, NumItems(InNum)
 	{}
 
-	uint64 Num() const					{ return NumItems;}
-	T operator[](uint64 Idx) const		{ check(Idx < NumItems); return Data[Idx]; }
-	const T* begin() const				{ return Data; }
-	const T* end() const				{ return Data + NumItems; }
+	inline uint64			Num() const								{ return NumItems;}
+	inline T				operator[](uint64 Idx) const			{ check(Idx < NumItems); return Data[Idx]; }
+	inline void				Copy(void* Dst, uint64 NumBytes) const;
+	inline const T*			begin() const							{ return Data; }
+	inline const T*			end() const								{ return Data + NumItems; }
 private:
 	const T* Data;
 	uint64 NumItems;
@@ -395,6 +398,26 @@ bool EqualItems(RangeTypeA&& A, RangeTypeB&& B)
 	return true;
 }
 
+inline void	FBoolRangeView::Copy(void* Dst, uint64 NumBytes) const
+{
+	check(NumBytes == NumBits);
+	bool* bOut = static_cast<bool*>(Dst);
+	for (bool b : *this)
+	{
+		*bOut++ = b;
+	}	
+}
+
+template<typename T>
+inline void TRangeView<T>::Copy(void* Dst, uint64 NumBytes) const
+{
+	check(NumBytes == NumItems * sizeof(T));
+	if (NumBytes)
+	{
+		FMemory::Memcpy(Dst, Data, NumBytes);
+	}
+}
+
 //////////////////////////////////////////////////////////////////////////
 
 class FLeafRangeView
@@ -451,6 +474,8 @@ public:
 	TRangeView<char16_t>	AsUtf16() const		{ return As<char16_t>(); }
 	TRangeView<char32_t>	AsUtf32() const		{ return As<char32_t>(); }
 
+	FLeafRangeLoadView		AsLoadView() const;
+
 private:
 	ELeafType						Type;
 	ELeafWidth						Width;
@@ -546,27 +571,27 @@ public:
 	explicit FMemberReader(FStructView Struct) : FMemberReader(Struct.Schema.Resolve(), Struct.Values, Struct.Schema.Batch) {}
 	FMemberReader(const FStructSchema& Schema, FByteReader Values, FReadBatchId InBatch);
 
-	bool					HasMore() const			{ return MemberIdx < NumMembers; }
+	bool								HasMore() const			{ return MemberIdx < NumMembers; }
 	
-	FOptionalMemberId		PeekName() const;		// @pre HasMore()
-	FOptionalMemberId		PeekNameUnchecked() const; // @pre HasMore()
-	EMemberKind				PeekKind() const;		// @pre HasMore()
-	FMemberType				PeekType() const;		// @pre HasMore()
+	PLAINPROPS_API FOptionalMemberId	PeekName() const;			// @pre HasMore()
+	PLAINPROPS_API FOptionalMemberId	PeekNameUnchecked() const; // @pre HasMore()
+	PLAINPROPS_API EMemberKind			PeekKind() const;			// @pre HasMore()
+	PLAINPROPS_API FMemberType			PeekType() const;			// @pre HasMore()
 
-	FLeafView				GrabLeaf();				// @pre PeekKind() == EMemberKind::Leaf
-	FRangeView				GrabRange();			// @pre PeekKind() == EMemberKind::Range
-	FStructView				GrabStruct();			// @pre PeekKind() == EMemberKind::Struct
-	//FAnyMemberView		GrabAny();				// @pre HasMore()
+	PLAINPROPS_API FLeafView			GrabLeaf();					// @pre PeekKind() == EMemberKind::Leaf
+	PLAINPROPS_API FRangeView			GrabRange();				// @pre PeekKind() == EMemberKind::Range
+	PLAINPROPS_API FStructView			GrabStruct();				// @pre PeekKind() == EMemberKind::Struct
+	//FAnyMemberView					GrabAny();					// @pre HasMore()
 
 	// Experimental 
 	// @pre Has N more contiguous members of the expected leaf type
 	template<Arithmetic T>
-	void					GrabLeaves(T* Out, uint32 N);
+	void								GrabLeaves(T* Out, uint32 N);
 
 	// Experimental
 	// @pre Has N more contiguous members of the expected leaf type
 	template<Enumeration T>
-	void					GrabEnums(T* Out, uint32 N);
+	void								GrabEnums(T* Out, uint32 N);
 
 protected: // for unit tests
 	const FMemberType*		Footer;
diff --git a/Engine/Plugins/Experimental/PlainProps/Source/Public/PlainPropsSave.h b/Engine/Plugins/Experimental/PlainProps/Source/Public/PlainPropsSave.h
--- a/Engine/Plugins/Experimental/PlainProps/Source/Public/PlainPropsSave.h
+++ b/Engine/Plugins/Experimental/PlainProps/Source/Public/PlainPropsSave.h
@@ -40,6 +40,7 @@ FSaveContext MakeSaveContext(FDefaultStructs Defaults, FScratchAllocator& Scratc
 
 [[nodiscard]] PLAINPROPS_API FBuiltStruct*	SaveStruct(const void* Struct, FStructSchemaId BindId, const FSaveContext& Context);
 [[nodiscard]] PLAINPROPS_API FBuiltStruct*	SaveStructDelta(const void* Struct, const void* Default, FStructSchemaId BindId, const FSaveContext& Context);
+[[nodiscard]] PLAINPROPS_API FBuiltStruct*	SaveStructDeltaIfDiff(const void* Struct, const void* Default, FStructSchemaId BindId, const FSaveContext& Context);
 [[nodiscard]] PLAINPROPS_API FBuiltRange*	SaveRange(const void* Range, FRangeMemberBinding Member, const FSaveContext& Ctx);
 [[nodiscard]] PLAINPROPS_API FBuiltRange*	SaveLeafRange(const void* Range, const ILeafRangeBinding& Binding, FUnpackedLeafType Leaf, const FSaveContext& Ctx);
 
diff --git a/Engine/Plugins/Experimental/PlainProps/Source/Public/PlainPropsSaveMember.h b/Engine/Plugins/Experimental/PlainProps/Source/Public/PlainPropsSaveMember.h
--- a/Engine/Plugins/Experimental/PlainProps/Source/Public/PlainPropsSaveMember.h
+++ b/Engine/Plugins/Experimental/PlainProps/Source/Public/PlainPropsSaveMember.h
@@ -273,11 +273,18 @@ using TMemberSerializer = typename TSelectSerializer<T>::Type;
 
 //////////////////////////////////////////////////////////////////////////
 
+// Helps hide FBuiltRange internals
 class FRangeSaverBase
 {
-protected:
+public:
 	PLAINPROPS_API FRangeSaverBase(FScratchAllocator& Scratch, uint64 Num, SIZE_T ItemSize);
-	~FRangeSaverBase() { check(!Range);	}
+
+protected:
+	[[nodiscard]] FTypedRange Finalize(FMemberSchema RangeSchema)
+	{
+		check(It == End);
+		return { RangeSchema, Range };
+	}
 
 	FBuiltRange*				Range;
 	uint8*						It;
@@ -303,7 +310,6 @@ public:
 
 	void AddItem(const T& Item)
 	{
-		check(Range);
 		check(It < End);
 		reinterpret_cast<BuiltItemType&>(*It) = Schema.SaveItem(Item, Ctx);
 		It += ItemSize;
@@ -311,11 +317,7 @@ public:
 
 	[[nodiscard]] FTypedRange Finalize(ERangeSizeType MaxSize)
 	{
-		check(Range);
-		check(It == End);
-		FTypedRange Out { Schema.MakeMemberRangeSchema(MaxSize), Range };
-		Range = nullptr;
-		return Out;
+		return FRangeSaverBase::Finalize(Schema.MakeMemberRangeSchema(MaxSize));
 	}
 
 private:	
@@ -323,4 +325,33 @@ private:
 	const FSaveContext&			Ctx;
 };
 
+// Saves an untyped homogeneous range - all items must have same type
+class FAnyRangeSaver : public FRangeSaverBase
+{
+public:
+	using FRangeSaverBase::FRangeSaverBase;
+	using FRangeSaverBase::Finalize;
+
+	inline void AddLeaf(const void* Item, SIZE_T ItemSize)
+	{
+		check(It < End);
+		FMemory::Memcpy(It, Item, ItemSize);
+		It += ItemSize;
+	}
+	
+	inline void AddRange(const FBuiltRange* Item)
+	{
+		check(It < End);
+		reinterpret_cast<const FBuiltRange*&>(*It) = Item;
+		It += sizeof(Item);
+	}
+
+	inline void AddStruct(const FBuiltStruct* Item)
+	{
+		check(It < End);
+		reinterpret_cast<const FBuiltStruct*&>(*It) = Item;
+		It += sizeof(Item);
+	}
+};
+
 } // namespace PlainProps
\ No newline at end of file
diff --git a/Engine/Plugins/Experimental/PlainProps/Source/Public/PlainPropsUeCoreBindings.h b/Engine/Plugins/Experimental/PlainProps/Source/Public/PlainPropsUeCoreBindings.h
--- a/Engine/Plugins/Experimental/PlainProps/Source/Public/PlainPropsUeCoreBindings.h
+++ b/Engine/Plugins/Experimental/PlainProps/Source/Public/PlainPropsUeCoreBindings.h
@@ -178,6 +178,7 @@ struct TStringBinding : ILeafRangeBinding
 	{
 		TArray<CharType>& Dst = static_cast<StringType*>(Range)->GetCharArray();
 		TRangeView<char8_t> Utf8 = Items.As<char8_t>();
+		// Todo: Better abstraction that hides internal representation
 		const UTF8CHAR* Src = reinterpret_cast<const UTF8CHAR*>(Utf8.begin());
 		int32 SrcLen = static_cast<int32>(Utf8.Num());
 		if (SrcLen == 0)
@@ -647,13 +648,13 @@ struct TBaseDeltaBinding : ICustomBinding, FBaseDeltaBinding
 	}
 
 	template<EOp Op, typename T>
-	void ApplyLeaves(SetType& Out, FLeafRangeView Items) const requires (Op == EOp::Add && !std::is_same_v<T, bool>)
+	void ApplyLeaves(SetType& Out, FLeafRangeLoadView Items) const requires (Op == EOp::Add && !std::is_same_v<T, bool>)
 	{
 		Out.Append(MakeArrayView(Items.As<T>()));
 	}
 
 	template<EOp Op, typename T>
-	void ApplyLeaves(SetType& Out, FLeafRangeView Items) const
+	void ApplyLeaves(SetType& Out, FLeafRangeLoadView Items) const
 	{
 		for (T Item : Items.As<T>())
 		{
diff --git a/Engine/Plugins/Experimental/PlainProps/Source/Tests/PlainPropsUeCoreTests.cpp b/Engine/Plugins/Experimental/PlainProps/Source/Tests/PlainPropsUeCoreTests.cpp
--- a/Engine/Plugins/Experimental/PlainProps/Source/Tests/PlainPropsUeCoreTests.cpp
+++ b/Engine/Plugins/Experimental/PlainProps/Source/Tests/PlainPropsUeCoreTests.cpp
@@ -175,7 +175,7 @@ template<class T>
 bool FBatchSaver::SaveDelta(const T& Object, const T& Default) 
 {
 	FDualStructSchemaId Id = IndexStructDualId<FIds, TTypename<std::remove_reference_t<T>>>();
-	if (FBuiltStruct* Delta = SaveStructDelta(&Object, &Default, Id.BindId, {GTypes, GSchemas, Customs, Scratch}))
+	if (FBuiltStruct* Delta = SaveStructDeltaIfDiff(&Object, &Default, Id.BindId, {GTypes, GSchemas, Customs, Scratch}))
 	{
 		SavedObjects.Emplace(Id.DeclId, MoveTemp(Delta));
 		return true;
diff --git a/Engine/Plugins/Experimental/PlainPropsUObject/Source/PlainPropsUObject.Build.cs b/Engine/Plugins/Experimental/PlainPropsUObject/Source/PlainPropsUObject.Build.cs
--- a/Engine/Plugins/Experimental/PlainPropsUObject/Source/PlainPropsUObject.Build.cs
+++ b/Engine/Plugins/Experimental/PlainPropsUObject/Source/PlainPropsUObject.Build.cs
@@ -14,6 +14,7 @@ public class PlainPropsUObject : ModuleRules
 			new string[] {
 				"Core",
 				"CoreUObject",
+				"CorePreciseFP",
 				"PlainProps",
 		});
 	}
diff --git a/Engine/Plugins/Experimental/PlainPropsUObject/Source/Private/PlainPropsUObjectBindings.cpp b/Engine/Plugins/Experimental/PlainPropsUObject/Source/Private/PlainPropsUObjectBindings.cpp
--- a/Engine/Plugins/Experimental/PlainPropsUObject/Source/Private/PlainPropsUObjectBindings.cpp
+++ b/Engine/Plugins/Experimental/PlainPropsUObject/Source/Private/PlainPropsUObjectBindings.cpp
@@ -5,6 +5,7 @@
 #include "Algo/Compare.h"
 #include "Containers/PagedArray.h"
 #include "Hash/xxhash.h"
+#include "Math/PreciseFP.h"
 #include "Misc/CommandLine.h"
 #include "Misc/CoreDelegates.h"
 #include "Misc/DefinePrivateMemberPtr.h"
@@ -36,9 +37,13 @@ FCommonUObjectIds::FCommonUObjectIds(TIdIndexer<FName>& Names)
 , SoftObjectPtr(	Names.IndexStruct({CoreUObject, Names.MakeTypename("SoftObjectPtr")}))
 , ScriptInterface(	Names.IndexStruct({CoreUObject, Names.MakeTypename("ScriptInterface")}))
 , Optional(			Names.NameType("Optional"))
+, Map(				Names.NameType("Map"))
+, Set(				Names.NameType("Set"))
 , Pair(				Names.NameType("Pair"))
 , Key(				Names.NameMember("Key"))
 , Value(			Names.NameMember("Value"))
+, Add(				Names.NameMember("Add"))
+, Del(				Names.NameMember("Del"))
 {}
 
 ////////////////////////////////////////////////////////////////////////////////////////////////
@@ -211,32 +216,25 @@ UE_DEFINE_PRIVATE_MEMBER_PTR(void(void*) const, GDestroyPropertyValue, FProperty
 namespace PlainProps::UE {
 
 
-inline void InitializeValue(FProperty* Property, void* Item)
+inline void InitializeValue(const FProperty* Property, void* Item)
 {
 	((*Property).*GInitPropertyValue)(Item);
 }
 
-inline void InitializeValues(FProperty* Property, uint8* Items, int32 Num, SIZE_T Stride)
+inline void InitializeValues(const FProperty* Property, uint8* Items, int32 Num, SIZE_T Stride)
 {
-	if (Property)
-	{
-		for (uint8* It = Items, *End = Items + NumBytes(Num, Stride); It != End; It += Stride)
-		{
-			((*Property).*GInitPropertyValue)(It);
-		}
-	}
-	else
+	for (uint8* It = Items, *End = Items + NumBytes(Num, Stride); It != End; It += Stride)
 	{
-		FMemory::Memzero(Items, NumBytes(Num, Stride));
+		((*Property).*GInitPropertyValue)(It);
 	}
 }
 
-inline void DestroyValue(FProperty* Property, void* Item)
+inline void DestroyValue(const FProperty* Property, void* Item)
 {
 	((*Property).*GDestroyPropertyValue)(Item);
 }
 
-inline void DestroyValues(FProperty* Property, uint8* Items, int32 Num, SIZE_T Stride)
+inline void DestroyValues(const FProperty* Property, uint8* Items, int32 Num, SIZE_T Stride)
 {
 	for (uint8* It = Items, *End = Items + NumBytes(Num, Stride); It != End; It += Stride)
 	{
@@ -347,7 +345,14 @@ struct TNonTrivialArrayBinding : TTrivialArrayBinding<ScriptArray>
 
 	void ConstructItems(int32 Num, uint8* Items) const
 	{
-		InitializeValues(Info.bConstructor ? Inner : nullptr, Items, Num, Info.ItemSize);
+		if (Info.bConstructor)
+		{
+			InitializeValues(Inner, Items, Num, Info.ItemSize);
+		}
+		else
+		{
+			FMemory::Memzero(Items, NumBytes(Num, Info.ItemSize));
+		}
 	}
 
 	void DestructItems(int32 Num, uint8* Items) const
@@ -378,7 +383,7 @@ struct TLeafArrayBinding : ILeafRangeBinding
 	{
 		FScriptArray& Array = *static_cast<FScriptArray*>(Range);
 		Array.SetNumUninitialized(static_cast<int32>(Leaves.Num()), LeafSize, LeafSize);
-		Leaves.Copy<Type, Width>(Array.GetData(), NumBytes(Array.Num()));
+		Leaves.AsBitCast<Type, Width>().Copy(Array.GetData(), NumBytes(Array.Num()));
 	}
 
 	virtual bool DiffLeaves(const void* RangeA, const void* RangeB) const override
@@ -531,19 +536,375 @@ inline void ReadSparseItems(FExistingItems& Dst, const FScriptSparseArray& Src,
 	}
 }
 
-struct FSetBinding : IItemRangeBinding
+// Helps save TSet/TMap deltas
+struct FSubSetIterator
 {
-	FProperty*					Element;
 	const FScriptSetLayout		Layout;
-	const bool					bDestruct;
-	const bool					bConstruct;
+	const FScriptSet&			Set;
+	const TBitArray<>&			Subset;
+	const int32					Max;
+	int32						Idx;
+
+	FSubSetIterator(const FScriptSetLayout& InLayout, const FScriptSet& InSet, const TBitArray<>& InSubset)
+	: Layout(InLayout)
+	, Set(InSet)
+	, Subset(InSubset)
+	, Max(InSet.GetMaxIndex())
+	, Idx(Max > 0 ? Subset.Find(true) : INDEX_NONE)
+	{}
+
+	explicit operator bool() const	{ return Idx != INDEX_NONE; }
+	const void* operator*() const	{ return Set.GetData(Idx, Layout); }
+	void operator++()				{ Idx = ++Idx < Max ? Subset.FindFrom(true, Idx) : INDEX_NONE; }
+	uint32 CountNum() const			{ return Subset.CountSetBits(); }
+};
+
+////////////////////////////////////////////////////////////////////////////////////////////////
+
+// Helpers to avoid using leaf FProperty instances after binding
+//
+// Below must match FFloatProperty, FDoubleProperty, FBoolProperty, FEnumProperty, TNumericProperty
+// Identical() and GetValueTypeHashInternal() implementations perfectly except not supporting nullptrs
+inline uint32	LeafPropertyHash(float In)											{ return ::UE::PreciseFPHash(In); }
+inline uint32	LeafPropertyHash(double In)											{ return ::UE::PreciseFPHash(In); }
+inline bool		LeafPropertyIdentical(float A, float B)								{ return ::UE::PreciseFPEqual(A, B); }
+inline bool		LeafPropertyIdentical(double A, double B)							{ return ::UE::PreciseFPEqual(A, B); }
+template<typename T>
+inline uint32	LeafPropertyHash(T In) requires (std::is_unsigned_v<T>)				{ return GetTypeHash(In); }
+template<typename T>
+inline bool		LeafPropertyIdentical(T A, T B) requires (std::is_unsigned_v<T>)	{ return A == B; }
+
+// Type-erased just enough to call LeafPropertyHash / LeafPropertyIdentical and FLeafRangeLoadView::As/AsBitcast
+enum class EPropertyKind : uint8 { Range, Struct, Bool, Integral8, Integral16, Integral32, Integral64, Float32, Float64 };
+
+template<EPropertyKind Kind>
+struct TLeafTypeOf											{ using Type = void; };
+template<> struct TLeafTypeOf<EPropertyKind::Bool>			{ using Type = bool; };
+template<> struct TLeafTypeOf<EPropertyKind::Integral8>		{ using Type = uint8; };
+template<> struct TLeafTypeOf<EPropertyKind::Integral16>	{ using Type = uint16; };
+template<> struct TLeafTypeOf<EPropertyKind::Integral32>	{ using Type = uint32; };
+template<> struct TLeafTypeOf<EPropertyKind::Integral64>	{ using Type = uint64; };
+template<> struct TLeafTypeOf<EPropertyKind::Float32>		{ using Type = float; };
+template<> struct TLeafTypeOf<EPropertyKind::Float64>		{ using Type = double; };
+
+template<EPropertyKind Kind>
+using LeafTypeOf = typename TLeafTypeOf<Kind>::Type;
+
+template<typename LeafType>
+uint32 LeafHash(const void* In)
+{
+	return LeafPropertyHash(*static_cast<const LeafType*>(In));
+}
+
+template<typename LeafType>
+bool LeafIdentical(const void* A, const void* B)
+{
+	return LeafPropertyIdentical(*static_cast<const LeafType*>(A), *static_cast<const LeafType*>(B));
+}
+
+template<typename LeafType>
+inline auto CastAs(FLeafRangeLoadView In)
+{
+	if constexpr (std::is_floating_point_v<LeafType>)
+	{
+		return In.As<LeafType>();
+	}
+	else
+	{
+		return In.AsBitCast<LeafType>();
+	}
+}
+
+inline EPropertyKind GetPropertyKind(FLeafBindType In)
+{
+	check(In.Bind.Type != ELeafBindType::BitfieldBool);
+	ELeafType Type = ToLeafType(In.Bind.Type);
+	ELeafWidth Width = In.Basic.Width;
+
+	if (Type == ELeafType::Float)
+	{
+		return Width == ELeafWidth::B32 ? EPropertyKind::Float32 : EPropertyKind::Float64;
+	}
+	else if (Type == ELeafType::Bool)
+	{
+		return EPropertyKind::Bool;
+	}
+	else switch (Width)
+	{
+		case ELeafWidth::B8:	return EPropertyKind::Integral8;
+		case ELeafWidth::B16:	return EPropertyKind::Integral16;
+		case ELeafWidth::B32:	return EPropertyKind::Integral32;
+		default:				return EPropertyKind::Integral64;
+	}								
+}
+
+static EPropertyKind GetPropertyKind(FMemberBinding In)
+{
+	return (In.RangeBindings.Num() > 0) ? EPropertyKind::Range
+										: In.InnermostType.IsStruct()
+										? EPropertyKind::Struct
+										: GetPropertyKind(In.InnermostType.AsLeaf());
+}
+
+////////////////////////////////////////////////////////////////////////////////////////////////
+
+inline uint32 HashRangeBindings(TConstArrayView<FRangeBinding> In)
+{
+	return static_cast<uint32>(FXxHash64::HashBuffer(In.GetData(), In.NumBytes()).Hash);
+}
+
+inline uint32 HashSkipOffset(FMemberBinding In)
+{
+	uint32 Out = HashCombineFast(GetTypeHash(In.InnermostSchema), GetTypeHash(In.InnermostType));
+	return In.RangeBindings.IsEmpty() ? Out : HashCombineFast(Out, HashRangeBindings(In.RangeBindings));
+}
+
+inline bool EqSkipOffset(FMemberBinding A, FMemberBinding B)
+{
+	return A.InnermostType == B.InnermostType && A.InnermostSchema == B.InnermostSchema && Algo::Compare(A.RangeBindings, B.RangeBindings);
+}
+
+// Instead of TMap KeyFuncs
+struct FParameterBinding : FMemberBinding
+{
+	friend uint32 GetTypeHash(FParameterBinding In) { return HashSkipOffset(In); };
+	inline bool operator==(FParameterBinding O) const { return EqSkipOffset(*this, O); }
+};
+
+////////////////////////////////////////////////////////////////////////////////////////////////
+
+inline bool operator==(FScriptSetLayout A, FScriptSetLayout B)
+{
+	return FMemory::Memcmp(&A, &B, sizeof(A)) == 0;
+}
+
+// Used for leaf sets
+template<EPropertyKind ItemKind, typename LeafType>
+struct TSetPropertyBase
+{
+	using SetElemT = TSetElement<LeafType>;
+
+	static constexpr int32							ElementSize = sizeof(LeafType);
+	static constexpr uint16							MaxLoadScratchItems = sizeof(FLoadRangeContext::Scratch) / ElementSize;
+	static constexpr FScriptSetLayout				Layout = {offsetof(SetElemT, HashNextId), offsetof(SetElemT, HashIndex), sizeof(SetElemT), {alignof(SetElemT), sizeof(SetElemT)}};
+	static constexpr bool							bDestruct = false;
+	static constexpr bool							bConstruct = false;
+	static constexpr SIZE_T							BuiltItemSize = ElementSize;
 	
-	explicit FSetBinding(FSetProperty* In)
-	: Element(In->ElementProp)
+	const FMemberBindType							InnerBindType;
+	const FMemberType								InnerType;
+	const FOptionalSchemaId							InnermostSchema;
+	const FRangeBinding								RangeBinding;					
+
+	TSetPropertyBase(FSetProperty* Set, FMemberBinding Elem, const IItemRangeBinding* This)
+	: InnerBindType(Elem.InnermostType)
+	, InnerType(ToLeafType(Elem.InnermostType.AsLeaf()))
+	, InnermostSchema(Elem.InnermostSchema)
+	, RangeBinding(*This, DefaultSizeType)
+	{
+		check(Layout == Set->SetLayout);
+		check(ElementSize == Set->ElementProp->GetElementSize());
+		check(ElementSize == SizeOf(InnerType.AsLeaf().Width));
+	}
+
+	static constexpr SIZE_T GetStride()								{ return sizeof(TSetElement<LeafType>); }
+	static constexpr void	InitCompactElems(void*, uint64, SIZE_T) {}
+	static constexpr void	DestructElems(FScriptSet& Set)			{}
+
+	static constexpr int32 FindIndex(const FScriptSet& Set, const void* Elem)
+	{
+		return Set.FindIndex(Elem, Layout, &LeafHash<LeafType>, &LeafIdentical<LeafType>);
+	}
+
+	static constexpr void Rehash(FScriptSet& Set)
+	{
+		Set.Rehash(Layout, &LeafHash<LeafType>);
+	}
+
+	inline FRangeMemberBinding MakeRangeMemberBinding() const
+	{
+		return { &InnerBindType, &RangeBinding, 1, InnermostSchema, 0 };
+	}
+
+	inline FMemberSchema MakeRangeSchema() const
+	{
+		return { FMemberType(DefaultSizeType), InnerType, 1, InnermostSchema, nullptr };
+	}
+
+	inline static void SaveDeltaItems(FAnyRangeSaver& Out, FSubSetIterator It, const FSaveContext&)
+	{
+		// Can be optimized by checking IsCompact() and adding FAnyRangeSaver::AddLeaves
+		for (; It; ++It)
+		{
+			Out.AddLeaf(*It, sizeof(LeafType));
+		}
+	}
+};
+
+// Used for non-leaf sets
+template<EPropertyKind ItemKind>
+struct TSetPropertyBase<ItemKind, void>
+{
+	static constexpr SIZE_T							BuiltItemSize = sizeof(FBuiltValue);
+
+	FProperty const* const							ElemProp;
+	const FScriptSetLayout							Layout;
+	const bool										bDestruct;
+	const bool										bConstruct;
+	const uint16									MaxLoadScratchItems;
+	const int32										ElementSize;
+	const FOptionalSchemaId							InnermostSchema;
+	TArray<FMemberType, TInlineAllocator<8>>		InnerTypes;
+	TArray<FMemberBindType, TInlineAllocator<8>>	InnerBindTypes;
+	TArray<FRangeBinding, TInlineAllocator<2>>		RangeBindings;
+
+	TSetPropertyBase(FSetProperty* In, FMemberBinding Elem, const IItemRangeBinding* This)
+	: ElemProp(In->ElementProp)
 	, Layout(In->SetLayout)
 	, bDestruct(HasDestructor(In->ElementProp))
 	, bConstruct(HasConstructor(In->ElementProp))
-	{}
+	, MaxLoadScratchItems(IntCastChecked<uint16>(sizeof(FLoadRangeContext::Scratch) / In->ElementProp->GetElementSize()))
+	, ElementSize(In->ElementProp->GetElementSize())
+	, InnermostSchema(Elem.InnermostSchema)
+	{
+		for (FRangeBinding Inner : Elem.RangeBindings)
+		{
+			InnerTypes.Emplace(Inner.GetSizeType());
+			InnerBindTypes.Emplace(Inner.GetSizeType());
+		}
+		InnerTypes.Emplace(Elem.InnermostType.IsStruct() ? FMemberType(Elem.InnermostType.AsStruct()) 
+														 : FMemberType(ToLeafType(Elem.InnermostType.AsLeaf())));
+		InnerBindTypes.Emplace(Elem.InnermostType);
+		RangeBindings.Emplace(*This, DefaultSizeType);
+		RangeBindings.Append(Elem.RangeBindings);
+		check(RangeBindings.Num() <= 0xFFFF);
+		check(Elem.RangeBindings.IsEmpty() == (ItemKind != EPropertyKind::Range));
+		check(ElemProp->GetElementSize() == ElemSize());
+	}
+
+	inline SIZE_T			GetStride() const					{ return Layout.Size; }
+	inline SIZE_T			ElemSize() const					{ return ElementSize; }
+	inline auto				MakeHashFn() const					{ return [P = ElemProp](const void* In) { return P->GetValueTypeHash(In); }; }
+	inline auto				MakeIdenticalFn() const				{ return [P = ElemProp](const void* A, const void* B) { return P->Identical(A, B); }; }
+
+	inline int32 FindIndex(const FScriptSet& Set, const void* Elem) const
+	{
+		return Set.FindIndex(Elem, Layout, MakeHashFn(), MakeIdenticalFn());
+	}
+
+	inline void Rehash(FScriptSet& Set) const
+	{
+		Set.Rehash(Layout, MakeHashFn());
+	}
+
+	inline FRangeMemberBinding MakeRangeMemberBinding() const
+	{
+		return { InnerBindTypes.GetData(), RangeBindings.GetData(), static_cast<uint16>(InnerBindTypes.Num()), InnermostSchema, 0 };
+	}
+
+	inline FMemberSchema MakeRangeSchema() const
+	{
+		return { FMemberType(DefaultSizeType), InnerTypes[0], static_cast<uint16>(InnerTypes.Num()), InnermostSchema, InnerTypes.Num() > 1 ? InnerTypes.GetData() : nullptr };
+	}
+
+	inline TConstArrayView<FRangeBinding> GetInnerBindings() const
+	{
+		return { RangeBindings.GetData() + 1, RangeBindings.Num() - 1 };
+	}
+
+	inline void* ConstructElem(void* Item) const
+	{
+		if (bConstruct)
+		{
+			InitializeValue(ElemProp, Item);
+		}
+		else
+		{
+			FMemory::Memzero(Item, ElemSize());
+		}
+		return Item;
+	}
+
+	inline void InitCompactElems(void* Items, uint64 Num, SIZE_T Stride) const
+	{
+		if (bConstruct)
+		{
+			InitializeValues(ElemProp, static_cast<uint8*>(Items), Num, Stride);
+		}
+		else
+		{
+			for (uint8* It = static_cast<uint8*>(Items), *End = It + Num * Stride; It != End; It += Stride)
+			{
+				FMemory::Memzero(It, ElemSize());
+			}
+		}
+	}
+
+	inline void DestructElem(void* Elem) const
+	{
+		if (bDestruct)
+		{
+			DestroyValue(ElemProp, Elem);
+		}
+	}
+
+	inline void DestructElems(FScriptSet& Set) const
+	{
+		check(bDestruct);
+		uint8* It = static_cast<uint8*>(Set.GetData(0, Layout));
+		SIZE_T Stride = GetStride();
+		if (IsCompact(Set))
+		{
+			DestroyValues(ElemProp, It, Set.Num(), Stride);
+		}
+		else
+		{
+			for (int32 Idx = 0, Max = Set.GetMaxIndex(); Idx < Max; ++Idx)
+			{
+				if (Set.IsValidIndex(Idx))
+				{
+					DestroyValue(ElemProp, It);
+				}
+				It += Stride;
+			}
+		}
+	}
+
+	inline void SaveDeltaItems(FAnyRangeSaver& Out, FSubSetIterator& It, const FSaveContext& Ctx) const requires (ItemKind == EPropertyKind::Range)
+	{
+		FRangeMemberBinding Inner = GetInnerRange(MakeRangeMemberBinding());
+		for (; It; ++It)
+		{
+			Out.AddRange(SaveRange(*It, Inner, Ctx));
+		}
+	}
+
+	inline void SaveDeltaItems(FAnyRangeSaver& Out, FSubSetIterator& It, const FSaveContext& Ctx) const requires (ItemKind == EPropertyKind::Struct)
+	{
+		FStructSchemaId Id = static_cast<FStructSchemaId>(InnermostSchema.Get());
+		// Todo: Lookup Ctx.Defaults and call SaveStructDelta instead
+		for (; It; ++It)
+		{
+			Out.AddStruct(SaveStruct(*It, Id, Ctx));
+		}
+	}
+};
+
+template<EPropertyKind ItemKind>
+struct TSetPropertyBinding : IItemRangeBinding, ICustomBinding, TSetPropertyBase<ItemKind, LeafTypeOf<ItemKind>>
+{
+	using LeafType = LeafTypeOf<ItemKind>;
+	using Base = TSetPropertyBase<ItemKind, LeafType>;
+	static constexpr bool bLeafItems = !std::is_void_v<LeafType>;
+	
+	TSetPropertyBinding(FSetProperty* In, FMemberBinding Elem) : Base(In, Elem, this) {}
+	
+	using Base::ElementSize;
+	using Base::MaxLoadScratchItems;
+	using Base::Layout;
+	using Base::bConstruct;
+	using Base::bDestruct;
+	using Base::GetStride;
 
 	virtual void ReadItems(FSaveRangeContext& Ctx) const override
 	{
@@ -551,64 +912,389 @@ struct FSetBinding : IItemRangeBinding
 	}
 
 	virtual void MakeItems(FLoadRangeContext& Ctx) const override
-	{
+	{		
+		// Load into scratch and memcpy to sparse TSetElement array to avoid FProperty::CopyValuesInternal dependency
+		// TSet elements are always trivially relocatable. Scratch items are constructed but not destroyed.
+		const bool bScratch = Ctx.Request.NumTotal() > 1 && MaxLoadScratchItems > 1;
+		const int32 NumTotal = static_cast<int32>(Ctx.Request.NumTotal());
+		const SIZE_T ItemSize = ElementSize;
+		const SIZE_T Stride = GetStride();
+
 		FScriptSet& Set = Ctx.Request.GetRange<FScriptSet>();
-		const int32 NewNum = static_cast<int32>(Ctx.Request.NumTotal());
-				
 		if (Ctx.Request.IsFirstCall())
 		{
 			if (bDestruct && !Set.IsEmpty())
 			{
-				DestructItems(Set);
+				Base::DestructElems(Set);
 			}
 
-			Set.Empty(NewNum, Layout);
+			Set.Empty(NumTotal, Layout);
 			
-			if (NewNum > 0)
+			if (NumTotal > 0)
 			{
-				// Todo: Add TScriptSparseArray::AppendUninitialized or SetNumUninitialized
-				for (int32 Idx = 0; Idx < NewNum; ++Idx)
+				// Todo: Add TScriptSparseArray::AppendUninitialized or SetNumUninitialized or Reset
+				for (int32 Idx = 0; Idx < NumTotal; ++Idx)
 				{
 					Set.AddUninitialized(Layout);
 				}
-				check(/* Set.IsCompact() */ Set.GetMaxIndex() == NewNum);
+				check(IsCompact(Set));
 
-				uint8* Data = static_cast<uint8*>(Set.GetData(0, Layout));
-				InitializeValues(bConstruct ? Element : nullptr, Data, NewNum, Layout.Size);
+				if (bScratch)
+				{
+					int32 NumScratch = FMath::Min(int32(MaxLoadScratchItems), NumTotal);
+					Base::InitCompactElems(Ctx.Scratch, NumScratch, ItemSize);
+					Ctx.Items.Set(Ctx.Scratch, NumScratch, ItemSize);
+				}
+				else // Initialize all elements and then load them one by one
+				{
+					uint8* Data = static_cast<uint8*>(Set.GetData(0, Layout));
+					Base::InitCompactElems(Data, NumTotal, Stride);
+					Ctx.Items.Set(Data, 1, ItemSize);
+				}
 
-				Ctx.Items.Set(Data, Ctx.Request.NumTotal(), Layout.Size);
 				Ctx.Items.RequestFinalCall();
 			}
 		}
+		else if (Ctx.Request.IsFinalCall())
+		{
+			Base::Rehash(Set);
+		}
 		else
 		{
-			check(Ctx.Request.IsFinalCall());
-			Set.Rehash(Layout, [this](const void* Src) { return Element->GetValueTypeHash(Src); });
+			uint8* DstIt = static_cast<uint8*>(Set.GetData(static_cast<int32>(Ctx.Request.GetIndex()), Layout));
+
+			if (bScratch)
+			{
+				// Move loaded items from scratch to sparse array
+				for (const uint8* SrcIt = Ctx.Items.GetData(), *SrcEnd = Ctx.Items.GetDataEnd(); SrcIt != SrcEnd; SrcIt += ItemSize)
+				{
+					FMemory::Memcpy(DstIt, SrcIt, ItemSize);
+					DstIt += Stride;
+				}
+
+				Ctx.Items.UpdateNum(FMath::Min(Ctx.Request.NumMore(), uint64(MaxLoadScratchItems)));
+			}
+			else // Load one by one
+			{
+				Ctx.Items.Set(DstIt, 1, ItemSize);
+			}
 		}
 	}
+	
+	inline bool HasElem(const FScriptSet& Set, const void* Elem) const
+	{
+		return Base::FindIndex(Set, Elem) != INDEX_NONE;
+	}
 
-	inline void DestructItems(FScriptSet& Set) const
+	inline void RemoveElem(FScriptSet& Set, const void* Elem) const
 	{
-		if (IsCompact(Set))
+		int32 Idx = Base::FindIndex(Set, Elem);
+		if (Idx != INDEX_NONE)
 		{
-			uint8* Data = static_cast<uint8*>(Set.GetData(0, Layout));
-			DestroyValues(Element, Data, Set.Num(), Layout.Size);
+			Set.RemoveAt(Idx, Layout);
+		}
+	}
+
+	virtual void SaveCustom(FMemberBuilder& Dst, const void* InSrc, const void* InDefault, const FSaveContext& Ctx) override
+	{
+		const FScriptSet& Src = *static_cast<const FScriptSet*>(InSrc);
+		const FScriptSet* Default = static_cast<const FScriptSet*>(InDefault);
+
+		if (Src.IsEmpty())
+		{
+			if (Default && !Default->IsEmpty())
+			{
+				FBuiltRange* Range = SaveRange(Default, Base::MakeRangeMemberBinding(), Ctx);
+				Dst.AddRange(GUE.CommonIds.Del, { Base::MakeRangeSchema(), Range });
+			}
+		}
+		else if (Default && !Default->IsEmpty())
+		{
+			TBitArray<> DeletedIds(false, Default->GetMaxIndex());
+			for (int32 Idx = 0, Max = Default->GetMaxIndex(); Idx < Max; ++Idx)
+			{
+				DeletedIds[Idx] = Default->IsValidIndex(Idx) && HasElem(Src, Default->GetData(Idx, Layout));
+			}
+			if (FSubSetIterator Deleted{Layout, *Default, DeletedIds})
+			{
+				Dst.AddRange(GUE.CommonIds.Del, SaveDeltaRange(Deleted, Ctx));
+			}
+			
+			TBitArray<> AddedIds(false, Src.GetMaxIndex());
+			for (int32 Idx = 0, Max = Src.GetMaxIndex(); Idx < Max; ++Idx)
+			{
+				AddedIds[Idx] = Src.IsValidIndex(Idx) && !HasElem(*Default, Src.GetData(Idx, Layout));
+			}
+			if (FSubSetIterator Added{Layout, Src, AddedIds})
+			{
+				Dst.AddRange(GUE.CommonIds.Add, SaveDeltaRange(Added, Ctx));
+			}
 		}
 		else
 		{
-			for (int32 Idx = 0, Num = Set.Num(); Idx < Num; ++Idx)
+			FBuiltRange* Range = SaveRange(&Src, Base::MakeRangeMemberBinding(), Ctx);
+			Dst.AddRange(GUE.CommonIds.Add, { Base::MakeRangeSchema(), Range });
+		}
+	}
+
+	virtual void LoadCustom(void* InDst, FStructLoadView Src, ECustomLoadMethod Method) const override
+	{
+		check(Method == ECustomLoadMethod::Assign);
+		FScriptSet& Dst = *static_cast<FScriptSet*>(InDst);
+		FMemberLoader Members(Src);				
+		while (Members.HasMore())
+		{
+			if (Members.PeekName() == GUE.CommonIds.Add)
 			{
-				DestroyValue(Element, Set.GetData(Idx, Layout));
+				if (Dst.IsEmpty())
+				{
+					Append(Dst, Members.GrabRange());
+				}
+				else
+				{
+					Insert(Dst, Members.GrabRange());
+				}
+				check(!Members.HasMore());
+				break;
+			}
+
+			check(Members.PeekName() == GUE.CommonIds.Del);
+			if (!Dst.IsEmpty())
+			{
+				Remove(Dst, Members.GrabRange());
+			}
+		}	
+	}
+
+	virtual bool DiffCustom(const void* InA, const void* InB) const override
+	{
+		const FScriptSet& A = *static_cast<const FScriptSet*>(InA);
+		const FScriptSet& B = *static_cast<const FScriptSet*>(InB);
+
+		if (A.NumUnchecked() != B.NumUnchecked())
+		{
+			return true;
+		}
+
+		if (A.NumUnchecked() > 0)
+		{
+			for (int32 IdxA = 0, MaxA = A.GetMaxIndex(); IdxA < MaxA; ++IdxA)
+			{
+				if (A.IsValidIndex(IdxA) && !HasElem(B, A.GetData(IdxA, Layout)))
+				{
+					return true;
+				}
+			}
+		}
+
+		return false;
+	}
+	
+	// Load into empty set
+	inline void Append(FScriptSet& Dst, FRangeLoadView Src) const
+	{
+		// There's no TScriptSparseArray::SetNumUninitialized() (yet), 
+		// reserve using Empty() and add items one by one instead
+		Dst.Empty(static_cast<int32>(Src.Num()), Layout);
+		for (uint64 Idx = 0; Idx < Src.Num(); ++Idx)
+		{
+			Dst.AddUninitialized(Layout);
+		}
+		check(IsCompact(Dst));
+		uint8* It = static_cast<uint8*>(Dst.GetData(0, Layout));
+		SIZE_T Stride = GetStride();
+		Base::InitCompactElems(It, Src.Num(), Stride);
+
+		if constexpr (bLeafItems)
+		{
+			for (LeafType Leaf : CastAs<LeafType>(Src.AsLeaves()))
+			{
+				*reinterpret_cast<LeafType*>(It) = Leaf;
+				It += Stride;
 			}
 		}
+		else if constexpr (ItemKind == EPropertyKind::Range)
+		{
+			TConstArrayView<FRangeBinding> Inners = Base::GetInnerBindings();
+			for (FRangeLoadView Range : Src.AsRanges())
+			{
+				LoadRange(It, Range, Inners);
+				It += Stride;
+			}
+		}
+		else
+		{
+			for (FStructLoadView Struct : Src.AsStructs())
+			{
+				LoadStruct(It, Struct);
+				It += Stride;
+			}
+		}
+
+		Base::Rehash(Dst);
+	}
+
+	
+	inline void Insert(FScriptSet& Dst, FRangeLoadView Src) const requires(bLeafItems)
+	{	
+		for (LeafType Leaf : CastAs<LeafType>(Src.AsLeaves()))
+		{
+			if (!HasElem(Dst, &Leaf))
+			{
+				 void* Elem = Dst.GetData(Dst.AddUninitialized(Layout), Layout);
+				 *static_cast<LeafType*>(Elem) = Leaf;
+			}
+		}
+
+		Base::Rehash(Dst);
+	}
+
+	inline void* AddConstructed(FScriptSparseArray& Dst, int32& OutIdx) const
+	{
+		OutIdx = Dst.AddUninitialized(Layout.SparseArrayLayout);
+		return Base::ConstructElem(Dst.GetData(OutIdx, Layout.SparseArrayLayout));
+	}
+
+	// Load structs or ranges into non-empty set
+	inline void Insert(FScriptSet& Dst, FRangeLoadView Src) const
+	{
+		// Written to avoid FProperty::CopyCompleteValue_InContainer dependency
+		// Items are loaded directly into sparse array and then removed if a duplicate existed
+
+		FScriptSparseArray& DstItems = reinterpret_cast<FScriptSparseArray&>(Dst);
+		const int32 OldNum = DstItems.NumUnchecked();
+		int32 TmpIdx;
+		void* Tmp = AddConstructed(DstItems, /* out*/ TmpIdx);
+		bool bSkip = true;
+		if constexpr (ItemKind == EPropertyKind::Range)
+		{
+			TConstArrayView<FRangeBinding> Inners = Base::GetInnerBindings();
+			for (FRangeLoadView Range : Src.AsRanges())
+			{
+				LoadRange(Tmp, Range, Inners);
+				bSkip = HasElem(Dst, Tmp);
+				Tmp = bSkip ? Tmp : AddConstructed(DstItems, /* out*/ TmpIdx);
+			}
+		}
+		else
+		{
+			for (FStructLoadView Struct : Src.AsStructs())
+			{
+				LoadStruct(Tmp, Struct);
+				bSkip = HasElem(Dst, Tmp);
+				Tmp = bSkip ? Tmp : AddConstructed(DstItems, /* out*/ TmpIdx);
+			}
+		}
+
+		if (bSkip)
+		{
+			Base::DestructElem(Tmp);
+			DstItems.RemoveAtUninitialized(Layout.SparseArrayLayout, TmpIdx, 1);
+		}
+
+		if (DstItems.NumUnchecked() != OldNum)
+		{
+			Base::Rehash(Dst);
+		}
+	}
+
+	inline void Remove(FScriptSet& Dst, FRangeLoadView Src) const requires(bLeafItems)
+	{
+		for (LeafType Item : CastAs<LeafType>(Src.AsLeaves()))
+		{
+			RemoveElem(Dst, &Item);
+		}
+	}
+
+	inline void Remove(FScriptSet& Dst, FRangeLoadView Src) const
+	{
+		TArray<uint8, TInlineAllocator<64>> Buffer;
+		Buffer.SetNumUninitialized(ElementSize);
+		void* Tmp = Base::ConstructElem(Buffer.GetData());
+
+		if constexpr (ItemKind == EPropertyKind::Range)
+		{
+			TConstArrayView<FRangeBinding> Inners = Base::GetInnerBindings();
+			for (FRangeLoadView Item : Src.AsRanges())
+			{
+				LoadRange(Tmp, Item, Inners);
+				RemoveElem(Dst, Tmp);
+			}
+		}
+		else
+		{
+			for (FStructLoadView Item : Src.AsStructs())
+			{
+				LoadStruct(Tmp, Item);
+				RemoveElem(Dst, Tmp);
+			}
+		}
+
+		Base::DestructElem(Tmp);
+	}
+	
+	FTypedRange SaveDeltaRange(FSubSetIterator It, const FSaveContext& Ctx) const
+	{
+		FAnyRangeSaver Out(Ctx.Scratch, It.CountNum(), Base::BuiltItemSize);
+		Base::SaveDeltaItems(Out, It, Ctx);
+		return Out.Finalize(Base::MakeRangeSchema());
 	}
 };
 
-static FRangeBinding AllocateSetBinding(FSetProperty* Property)
+////////////////////////////////////////////////////////////////////////////////////////////////
+
+class FSetBindings
 {
-	// Todo: Ownership / memory leak
-	return FRangeBinding(*new FSetBinding(Property), DefaultSizeType);
-}
+	TMap<FParameterBinding, FRangeBinding> RangeBindings;
+
+	template<EPropertyKind ItemKind>
+	IItemRangeBinding* BindNew(FSetProperty* Property, FMemberBinding Elem, FStructSchemaId Id)
+	{
+		// Todo: Ownership / memory leak
+		TSetPropertyBinding<ItemKind>* Out = new TSetPropertyBinding<ItemKind>(Property, Elem);
+		GUE.DeltaCustoms.BindStruct(Id, Id, {}, *Out);
+		return Out;
+	}
+
+	IItemRangeBinding* DeltaBindNew(FSetProperty* Property, FMemberBinding Elem, FStructSchemaId Id)
+	{
+		switch (GetPropertyKind(Elem))
+		{
+		case EPropertyKind::Range:			return BindNew<EPropertyKind::Range>(Property, Elem, Id);
+		case EPropertyKind::Struct:			return BindNew<EPropertyKind::Struct>(Property, Elem, Id);
+		case EPropertyKind::Bool:			return BindNew<EPropertyKind::Bool>(Property, Elem, Id);
+		case EPropertyKind::Integral8:		return BindNew<EPropertyKind::Integral8>(Property, Elem, Id);
+		case EPropertyKind::Integral16:		return BindNew<EPropertyKind::Integral16>(Property, Elem, Id);
+		case EPropertyKind::Integral32:		return BindNew<EPropertyKind::Integral32>(Property, Elem, Id);
+		case EPropertyKind::Integral64:		return BindNew<EPropertyKind::Integral64>(Property, Elem, Id);
+		case EPropertyKind::Float32:		return BindNew<EPropertyKind::Float32>(Property, Elem, Id);
+		case EPropertyKind::Float64:		return BindNew<EPropertyKind::Float64>(Property, Elem, Id);
+		default: check(false);				return nullptr;
+		}
+	}
+public:
+
+	FRangeBinding Bind(FSetProperty* Property, FMemberBinding Elem)
+	{
+		check(Elem.Offset == 0);
+		if (const FRangeBinding* RangeBinding = RangeBindings.Find(FParameterBinding(Elem)))
+		{
+			return *RangeBinding;
+		}
+
+		// Index custom delta binding struct name
+		FTypeId Params[1] = { Elem.IndexParameterName(GUE.Names) };
+		FTypeId Type = { GUE.CommonIds.Core, FTypenameId(GUE.Names.MakeParametricTypeId(GUE.CommonIds.Set, Params)) };
+		FStructSchemaId Id = GUE.Names.IndexStruct(Type);
+		FMemberId Members[2] = { GUE.CommonIds.Del, GUE.CommonIds.Add };
+
+		// Todo: Ownership / memory leak
+		GUE.Types.DeclareStruct(Id, Type, Members, EMemberPresence::AllowSparse);
+		IItemRangeBinding* Out = DeltaBindNew(Property, Elem, Id);
+		return RangeBindings.Emplace(Elem, FRangeBinding(*Out, DefaultSizeType));
+	}
+};
+static FSetBindings GSets;
 
 ////////////////////////////////////////////////////////////////////////////////////////////////
 
@@ -651,8 +1337,9 @@ struct TMapPropertyBinding : IItemRangeBinding
 			Map.Empty(NewNum, Layout);
 			if (NewNum)
 			{
-				void* Data = AddAndInitItems(Map, NewNum);
-				Ctx.Items.Set(Data, Ctx.Request.NumTotal(), Layout.SetLayout.Size);
+				// Todo: Bug, mimic TSetPropertyBinding::MakeItems
+				void* Data = AddConstructedItems(Map, NewNum);
+				Ctx.Items.Set(Data, Ctx.Request.NumTotal(), ItemSize());
 				Ctx.Items.RequestFinalCall();
 			}
 		}
@@ -663,7 +1350,7 @@ struct TMapPropertyBinding : IItemRangeBinding
 		}
 	}
 
-	void* AddAndInitItems(ScriptMap& Map, const int32 Num) const
+	void* AddConstructedItems(ScriptMap& Map, const int32 Num) const
 	{
 		for (int32 Idx = 0; Idx < Num; ++Idx)
 		{
@@ -672,7 +1359,7 @@ struct TMapPropertyBinding : IItemRangeBinding
 		check(IsCompact(Map));
 
 		uint8* Data = static_cast<uint8*>(Map.GetData(0, Layout));
-		SIZE_T Stride = static_cast<SIZE_T>(Layout.SetLayout.Size);
+		SIZE_T Stride = ItemSize();
 		if constexpr (bConstructKey && bConstructValue == false)
 		{
 			FMemory::Memzero(Data, Stride * Num);
@@ -700,11 +1387,11 @@ struct TMapPropertyBinding : IItemRangeBinding
 					uint8* Data = static_cast<uint8*>(Map.GetData(0, Layout));
 					if constexpr (bDestructKey)
 					{
-						DestroyValues(KeyProp, Data, Num, Layout.SetLayout.Size);
+						DestroyValues(KeyProp, Data, Num, ItemSize());
 					}
 					if constexpr (bDestructValue)
 					{
-						DestroyValues(ValueProp, Data + Layout.ValueOffset, Num, Layout.SetLayout.Size);
+						DestroyValues(ValueProp, Data + Layout.ValueOffset, Num, ItemSize());
 					}
 				}
 				else
@@ -725,6 +1412,8 @@ struct TMapPropertyBinding : IItemRangeBinding
 			}
 		}
 	}
+
+	inline SIZE_T ItemSize() const { return Layout.SetLayout.SparseArrayLayout.Size; }
 };
 
 inline uint32 GetMapPropertyFlags(FProperty* Key, FProperty* Val)
@@ -771,22 +1460,6 @@ static FRangeBinding AllocateMapBinding(FMapProperty* Property)
 
 ////////////////////////////////////////////////////////////////////////////////////////////////
 
-inline uint32 HashRangeBindings(TConstArrayView<FRangeBinding> In)
-{
-	return static_cast<uint32>(FXxHash64::HashBuffer(In.GetData(), In.NumBytes()).Hash);
-}
-
-inline uint32 HashSkipOffset(FMemberBinding In)
-{
-	uint32 Out = HashCombineFast(GetTypeHash(In.InnermostSchema), GetTypeHash(In.InnermostType));
-	return In.RangeBindings.IsEmpty() ? Out : HashCombineFast(Out, HashRangeBindings(In.RangeBindings));
-}
-
-inline bool EqSkipOffset(FMemberBinding A, FMemberBinding B)
-{
-	return A.InnermostType == B.InnermostType && A.InnermostSchema == B.InnermostSchema && Algo::Compare(A.RangeBindings, B.RangeBindings);
-}
-
 class FPairBindings
 {
 	struct FPair
@@ -837,15 +1510,9 @@ static FPairBindings GPairs;
 
 class FOptionalBindings
 {
-	struct FValueBinding : FMemberBinding
-	{
-		friend uint32 GetTypeHash(FValueBinding In) { return HashSkipOffset(In); };
-		inline bool operator==(FValueBinding O) const { return EqSkipOffset(*this, O); }
-	};
-
-	TMap<FValueBinding, FStructSchemaId>	Bindings;
+	TMap<FParameterBinding, FStructSchemaId>	Bindings;
 
-	inline FStructSchemaId BindImpl(const FValueBinding Value)
+	inline FStructSchemaId BindImpl(const FParameterBinding Value)
 	{
 		check(Value.Offset == 0);
 		if (const FStructSchemaId* Id = Bindings.Find(Value))
@@ -873,7 +1540,7 @@ public:
 	inline FMemberBinding Bind(FMemberBinding Value)
 	{
 		FMemberBinding Out(0);
-		Out.InnermostSchema = BindImpl(static_cast<FValueBinding>(Value));
+		Out.InnermostSchema = BindImpl(static_cast<FParameterBinding>(Value));
 		Out.InnermostType = DefaultStructBindType;
 		return Out;
 	}
@@ -1083,14 +1750,14 @@ private:
 	}
 	inline FMemberBinding BindSet(FSetProperty* Property)
 	{
-		return BindAsRange(Property, AllocateSetBinding(Property), BindSingleProperty(Property->ElementProp));
+		FMemberBinding ElemBinding = BindSingleProperty(Property->ElementProp);
+		return BindAsRange(Property, GSets.Bind(Property, ElemBinding), ElemBinding);
 	}
 	inline FMemberBinding BindOptional(FOptionalProperty* Property)
 	{
 		return GOptionals.Bind(BindSingleProperty(Property->GetValueProperty()));
 	}
 
-
 	#if WITH_VERSE_VM
 	inline FMemberBinding BindVValue(FVValueProperty* Property) { return Todo(Property); }
 	inline FMemberBinding BindVRestValue(FVRestValueProperty* Property) { return Todo(Property); }
@@ -1211,7 +1878,6 @@ private:
 			Out.InnermostSchema = static_cast<FSchemaId>(StaticArrayId);
 			Out.RangeBindings = {};
 		}
-		
 	
 		return Out;	
 	}
diff --git a/Engine/Plugins/Experimental/PlainPropsUObject/Source/Public/PlainPropsUObjectRuntime.h b/Engine/Plugins/Experimental/PlainPropsUObject/Source/Public/PlainPropsUObjectRuntime.h
--- a/Engine/Plugins/Experimental/PlainPropsUObject/Source/Public/PlainPropsUObjectRuntime.h
+++ b/Engine/Plugins/Experimental/PlainPropsUObject/Source/Public/PlainPropsUObjectRuntime.h
@@ -29,9 +29,13 @@ struct FCommonUObjectIds
 	FStructSchemaId				ScriptInterface;
 	
 	FConcreteTypenameId			Optional;
+	FConcreteTypenameId			Map;
+	FConcreteTypenameId			Set;
 	FConcreteTypenameId			Pair;
 	FMemberId					Key;
 	FMemberId					Value;
+	FMemberId					Add;
+	FMemberId					Del;	
 };
 
 struct FGlobals
@@ -76,4 +80,4 @@ struct FDeltaRuntime : FDefaultRuntime
 	static FCustomBindings&		GetCustoms()		{ return GUE.DeltaCustoms; }
 };
 
-} // namespace PlainProps::UE
+} // namespace PlainProps::UE
\ No newline at end of file
