PlainProps add FMapProperty delta bindings and unify flat and delta bindings for TSet and TMap properties

#rnx

diff --git a/Engine/Plugins/Experimental/PlainProps/Source/Private/PlainPropsLoad.cpp b/Engine/Plugins/Experimental/PlainProps/Source/Private/PlainPropsLoad.cpp
--- a/Engine/Plugins/Experimental/PlainProps/Source/Private/PlainPropsLoad.cpp
+++ b/Engine/Plugins/Experimental/PlainProps/Source/Private/PlainPropsLoad.cpp
@@ -824,20 +824,27 @@ public:
 		
 	static void CopyRangeValues(const FConstructedItems& Items, FByteReader& ByteIt, FBitCacheReader&, const FLoadBatch& Batch, FUnpackedLeafType Leaf)
 	{
-		check(Items.Size == SizeOf(Leaf.Width));
-		if (Leaf.Type != ELeafType::Bool)
+		if (Items.Size == SizeOf(Leaf.Width))
 		{
-			FMemory::Memcpy(Items.Data, ByteIt.GrabBytes(Items.NumBytes()), Items.NumBytes());
-		}
-		else
-		{
-			FBoolRangeView Bits(ByteIt.GrabBytes(Align(Items.Num, 8)/8), Items.Num);
-			uint8* It = Items.Data;
-			for (bool bBit : Bits)
+			if (Leaf.Type != ELeafType::Bool)
+			{
+				FMemory::Memcpy(Items.Data, ByteIt.GrabBytes(Items.NumBytes()), Items.NumBytes());
+			}
+			else
 			{
-				reinterpret_cast<bool&>(*It++) = bBit;
+				FBoolRangeView Bits(ByteIt.GrabBytes(Align(Items.Num, 8)/8), Items.Num);
+				uint8* It = Items.Data;
+				for (bool bBit : Bits)
+				{
+					reinterpret_cast<bool&>(*It++) = bBit;
+				}
 			}
 		}
+		else // Strided
+		{
+			check(Items.Size > SizeOf(Leaf.Width));
+			check(false); // todo
+		}
 	}
 	
 	static void CopyRangeValues(const FConstructedItems& Items, FByteReader& ByteIt, FBitCacheReader&, const FLoadBatch& Batch, FStructSchemaId Id)
@@ -1122,6 +1129,15 @@ void LoadRange(void* Dst, FRangeLoadView Src, TConstArrayView<FRangeBinding> Bin
 	FRangeLoader::LoadView(static_cast<uint8*>(Dst), Src, Bindings);
 }
 
+void LoadRange(void* Dst, FByteReader& SrcBytes, FBitCacheReader& SrcBits, ERangeSizeType MaxSize, FRangeLoadSchema Schema, TConstArrayView<FRangeBinding> Bindings)
+{
+	if (uint64 Num = GrabRangeNum(MaxSize, SrcBytes, SrcBits))
+	{
+		FMemoryView Values = GrabRangeValues(Num, Schema.ItemType, SrcBytes);
+		LoadRange(Dst, {Schema, Num, Values}, Bindings);
+	}
+}
+
 void LoadStruct(void* Dst, FStructLoadView Src)
 {
 	LoadStruct(Dst, Src.Values, Src.Schema.LoadId, Src.Schema.Batch);
diff --git a/Engine/Plugins/Experimental/PlainProps/Source/Private/PlainPropsSave.cpp b/Engine/Plugins/Experimental/PlainProps/Source/Private/PlainPropsSave.cpp
--- a/Engine/Plugins/Experimental/PlainProps/Source/Private/PlainPropsSave.cpp
+++ b/Engine/Plugins/Experimental/PlainProps/Source/Private/PlainPropsSave.cpp
@@ -123,9 +123,9 @@ struct TNonLeafRangeSaver
 	}
 };
 
-using FNestedRangeSaver = TNonLeafRangeSaver<FBuiltRange*, FRangeMemberBinding>;
-using FStructRangeSaver = TNonLeafRangeSaver<FBuiltStruct*, FStructSchemaId>;
-using FStructRangeDeltaSaver = TNonLeafRangeSaver<FBuiltStruct*, FDefaultStruct>;
+using FInternalNestedRangeSaver = TNonLeafRangeSaver<FBuiltRange*, FRangeMemberBinding>;
+using FInternalStructRangeSaver = TNonLeafRangeSaver<FBuiltStruct*, FStructSchemaId>;
+using FInternalStructRangeDeltaSaver = TNonLeafRangeSaver<FBuiltStruct*, FDefaultStruct>;
 
 //////////////////////////////////////////////////////////////////////////
 
@@ -200,11 +200,11 @@ FBuiltRange* SaveLeafRange(const void* Range, const ILeafRangeBinding& Binding,
 {
 	if (const FDefaultStruct* Defaults = Algo::FindBy(Ctx.Defaults, Id, &FDefaultStruct::Id))
 	{
-		return SaveNonLeafRange<FStructRangeDeltaSaver>(Range, ItemBinding, Ctx, *Defaults);
+		return SaveNonLeafRange<FInternalStructRangeDeltaSaver>(Range, ItemBinding, Ctx, *Defaults);
 	}
 	else
 	{
-		return SaveNonLeafRange<FStructRangeSaver>(Range, ItemBinding, Ctx, Id);
+		return SaveNonLeafRange<FInternalStructRangeSaver>(Range, ItemBinding, Ctx, Id);
 	}
 }
 
@@ -228,7 +228,7 @@ static ELeafWidth GetBasicWidth(FLeafBindType Leaf)
 	switch (InnerType.GetKind())
 	{
 	case EMemberKind::Leaf:		return SaveLeafRange(Range, ItemBinding, Ctx, GetBasicWidth(InnerType.AsLeaf()));
-	case EMemberKind::Range:	return SaveNonLeafRange<FNestedRangeSaver>(Range, ItemBinding, Ctx, GetInnerRange(Member));
+	case EMemberKind::Range:	return SaveNonLeafRange<FInternalNestedRangeSaver>(Range, ItemBinding, Ctx, GetInnerRange(Member));
 	case EMemberKind::Struct:	return SaveStructRange(Range, ItemBinding, Ctx, static_cast<FStructSchemaId>(Member.InnermostSchema.Get()));
 	}
 
diff --git a/Engine/Plugins/Experimental/PlainProps/Source/Public/PlainPropsLoadMember.h b/Engine/Plugins/Experimental/PlainProps/Source/Public/PlainPropsLoadMember.h
--- a/Engine/Plugins/Experimental/PlainProps/Source/Public/PlainPropsLoadMember.h
+++ b/Engine/Plugins/Experimental/PlainProps/Source/Public/PlainPropsLoadMember.h
@@ -66,6 +66,8 @@ private:
 };
 
 PLAINPROPS_API void	LoadRange(void* Dst, FRangeLoadView Src, TConstArrayView<FRangeBinding> InnerBindings);
+// Experimental low-level API bypassing FMemberLoader, internal use only
+PLAINPROPS_API void	LoadRange(void* Dst, FByteReader& SrcBytes, FBitCacheReader& SrcBits, ERangeSizeType MaxSize, FRangeLoadSchema Schema, TConstArrayView<FRangeBinding> InnerBindings);
 
 //////////////////////////////////////////////////////////////////////////
 
diff --git a/Engine/Plugins/Experimental/PlainProps/Source/Public/PlainPropsSaveMember.h b/Engine/Plugins/Experimental/PlainProps/Source/Public/PlainPropsSaveMember.h
--- a/Engine/Plugins/Experimental/PlainProps/Source/Public/PlainPropsSaveMember.h
+++ b/Engine/Plugins/Experimental/PlainProps/Source/Public/PlainPropsSaveMember.h
@@ -280,6 +280,14 @@ public:
 	PLAINPROPS_API FRangeSaverBase(FScratchAllocator& Scratch, uint64 Num, SIZE_T ItemSize);
 
 protected:
+	template<typename BuiltItemType>
+	inline void AddBuiltItem(BuiltItemType Item)
+	{
+		check(It < End);
+		reinterpret_cast<BuiltItemType&>(*It) = Item;
+		It += sizeof(Item);
+	}
+
 	[[nodiscard]] FTypedRange Finalize(FMemberSchema RangeSchema)
 	{
 		check(It == End);
@@ -310,9 +318,7 @@ public:
 
 	void AddItem(const T& Item)
 	{
-		check(It < End);
-		reinterpret_cast<BuiltItemType&>(*It) = Schema.SaveItem(Item, Ctx);
-		It += ItemSize;
+		FRangeSaverBase::AddBuiltItem(Schema.SaveItem(Item, Ctx));
 	}
 
 	[[nodiscard]] FTypedRange Finalize(ERangeSizeType MaxSize)
@@ -325,33 +331,29 @@ private:
 	const FSaveContext&			Ctx;
 };
 
-// Saves an untyped homogeneous range - all items must have same type
-class FAnyRangeSaver : public FRangeSaverBase
+template<Arithmetic LeafType>
+class TLeafRangeSaver : public FRangeSaverBase
 {
 public:
-	using FRangeSaverBase::FRangeSaverBase;
+	TLeafRangeSaver(FScratchAllocator& Scratch, uint64 Num) : FRangeSaverBase(Scratch, Num, sizeof(LeafType)) {}
+	inline void AddItem(LeafType Item) {FRangeSaverBase::AddBuiltItem(Item); }
 	using FRangeSaverBase::Finalize;
+};
 
-	inline void AddLeaf(const void* Item, SIZE_T ItemSize)
-	{
-		check(It < End);
-		FMemory::Memcpy(It, Item, ItemSize);
-		It += ItemSize;
-	}
-	
-	inline void AddRange(const FBuiltRange* Item)
-	{
-		check(It < End);
-		reinterpret_cast<const FBuiltRange*&>(*It) = Item;
-		It += sizeof(Item);
-	}
+class FNestedRangeSaver : public FRangeSaverBase
+{
+public:
+	FNestedRangeSaver(FScratchAllocator& Scratch, uint64 Num) : FRangeSaverBase(Scratch, Num, sizeof(FBuiltRange*)) {}
+	inline void AddItem(const FBuiltRange* Item) {	FRangeSaverBase::AddBuiltItem(Item); }
+	using FRangeSaverBase::Finalize;
+};
 
-	inline void AddStruct(const FBuiltStruct* Item)
-	{
-		check(It < End);
-		reinterpret_cast<const FBuiltStruct*&>(*It) = Item;
-		It += sizeof(Item);
-	}
+class FStructRangeSaver : public FRangeSaverBase
+{
+public:
+	FStructRangeSaver(FScratchAllocator& Scratch, uint64 Num) : FRangeSaverBase(Scratch, Num, sizeof(FBuiltStruct*)) {}
+	inline void AddItem(const FBuiltStruct* Item) {	FRangeSaverBase::AddBuiltItem(Item); }
+	using FRangeSaverBase::Finalize;
 };
 
 } // namespace PlainProps
\ No newline at end of file
diff --git a/Engine/Plugins/Experimental/PlainPropsUObject/Source/Private/PlainPropsUObjectBindings.cpp b/Engine/Plugins/Experimental/PlainPropsUObject/Source/Private/PlainPropsUObjectBindings.cpp
--- a/Engine/Plugins/Experimental/PlainPropsUObject/Source/Private/PlainPropsUObjectBindings.cpp
+++ b/Engine/Plugins/Experimental/PlainPropsUObject/Source/Private/PlainPropsUObjectBindings.cpp
@@ -147,6 +147,31 @@ static FStructSchemaId DeclareStruct(const UStruct* Struct)
 
 ////////////////////////////////////////////////////////////////////////////////////////////////
 
+inline uint32 HashRangeBindings(TConstArrayView<FRangeBinding> In)
+{
+	return static_cast<uint32>(FXxHash64::HashBuffer(In.GetData(), In.NumBytes()).Hash);
+}
+
+inline uint32 HashSkipOffset(FMemberBinding In)
+{
+	uint32 Out = HashCombineFast(GetTypeHash(In.InnermostSchema), GetTypeHash(In.InnermostType));
+	return In.RangeBindings.IsEmpty() ? Out : HashCombineFast(Out, HashRangeBindings(In.RangeBindings));
+}
+
+inline bool EqSkipOffset(FMemberBinding A, FMemberBinding B)
+{
+	return A.InnermostType == B.InnermostType && A.InnermostSchema == B.InnermostSchema && Algo::Compare(A.RangeBindings, B.RangeBindings);
+}
+
+// Helper to cache various property bindings instead of a TMap KeyFunc
+struct FParameterBinding : FMemberBinding
+{
+	friend uint32 GetTypeHash(FParameterBinding In) { return HashSkipOffset(In); };
+	inline bool operator==(FParameterBinding O) const { return EqSkipOffset(*this, O); }
+};
+
+////////////////////////////////////////////////////////////////////////////////////////////////
+
 static bool IsMax(FName EnumName, FName ValueName)
 {
 	if (ValueName.GetNumber() != NAME_NO_NUMBER_INTERNAL)
@@ -209,34 +234,52 @@ inline bool HasDestructor(FProperty* Property)
 	return !(Property->PropertyFlags & (CPF_IsPlainOldData | CPF_NoDestructor));
 }
 
+inline bool HasHash(FProperty* Property)
+{
+	return !!(Property->PropertyFlags & CPF_HasGetValueTypeHash);
+}
+
 // Temp hack. Long-term, try throwing away FProperty and implement Construct/DestroyValue replacement
 } // namespace PlainProps::UE
 UE_DEFINE_PRIVATE_MEMBER_PTR(void(void*) const, GInitPropertyValue, FProperty, InitializeValueInternal);
 UE_DEFINE_PRIVATE_MEMBER_PTR(void(void*) const, GDestroyPropertyValue, FProperty, DestroyValueInternal);
+UE_DEFINE_PRIVATE_MEMBER_PTR(uint32(const void*) const, GHashPropertyValue, FProperty, GetValueTypeHashInternal);
 namespace PlainProps::UE {
 
+inline void ConstructValue(const FProperty* Property, void* Value)
+{
+	((*Property).*GInitPropertyValue)(Value);
+}
+
+inline void DestroyValue(const FProperty* Property, void* Value)
+{
+	((*Property).*GDestroyPropertyValue)(Value);
+}
 
-inline void InitializeValue(const FProperty* Property, void* Item)
+inline uint32 HashValue(const FProperty* Property, const void* Item)
 {
-	((*Property).*GInitPropertyValue)(Item);
+	return ((*Property).*GHashPropertyValue)(Item);
 }
 
-inline void InitializeValues(const FProperty* Property, uint8* Items, int32 Num, SIZE_T Stride)
+inline void ConstructValues(const FProperty* Property, uint8* Values, int32 Num, SIZE_T Stride)
 {
-	for (uint8* It = Items, *End = Items + NumBytes(Num, Stride); It != End; It += Stride)
+	for (uint8* It = Values, *End = Values + Num*Stride; It != End; It += Stride)
 	{
 		((*Property).*GInitPropertyValue)(It);
 	}
 }
 
-inline void DestroyValue(const FProperty* Property, void* Item)
+inline void MemzeroStrided(uint8* Values, int32 Num, SIZE_T Size, SIZE_T Stride)
 {
-	((*Property).*GDestroyPropertyValue)(Item);
+	for (uint8* It = Values, *End = Values + Num*Stride; It != End; It += Stride)
+	{
+		FMemory::Memzero(It, Size);
+	}
 }
 
-inline void DestroyValues(const FProperty* Property, uint8* Items, int32 Num, SIZE_T Stride)
+inline void DestroyValues(const FProperty* Property, uint8* Values, int32 Num, SIZE_T Stride)
 {
-	for (uint8* It = Items, *End = Items + NumBytes(Num, Stride); It != End; It += Stride)
+	for (uint8* It = Values, *End = Values + Num*Stride; It != End; It += Stride)
 	{
 		((*Property).*GDestroyPropertyValue)(It);
 	}
@@ -332,7 +375,7 @@ struct TNonTrivialArrayBinding : TTrivialArrayBinding<ScriptArray>
 		if (Diff > 0)
 		{
 			int32 OldNum = Array.Add(Diff, Info.ItemSize, Info.ItemAlign);
-			ConstructItems(Diff, static_cast<uint8*>(Array.GetData()) + NumBytes(OldNum, Info.ItemSize));
+			InitItems(Diff, static_cast<uint8*>(Array.GetData()) + NumBytes(OldNum, Info.ItemSize));
 		}
 		else if (Diff < 0)
 		{
@@ -343,11 +386,11 @@ struct TNonTrivialArrayBinding : TTrivialArrayBinding<ScriptArray>
 		Ctx.Items.Set(Array.GetData(), static_cast<uint64>(NewNum), Info.ItemSize);
 	}
 
-	void ConstructItems(int32 Num, uint8* Items) const
+	void InitItems(int32 Num, uint8* Items) const
 	{
 		if (Info.bConstructor)
 		{
-			InitializeValues(Inner, Items, Num, Info.ItemSize);
+			ConstructValues(Inner, Items, Num, Info.ItemSize);
 		}
 		else
 		{
@@ -494,73 +537,6 @@ static FRangeBinding AllocateArrayBinding(FArrayProperty* Property)
 
 ////////////////////////////////////////////////////////////////////////////////////////////////
 
-template<typename ScriptSet>
-inline bool IsCompact(const ScriptSet& Set)
-{
-	return Set.NumUnchecked() == Set.GetMaxIndex();
-}
-
-inline FExistingItemSlice GetContiguousSlice(int32 Idx, const FScriptSparseArray& Elems, const uint8* Data, SIZE_T Stride)
-{
-	int32 Num = 1;
-	for (;!Elems.IsValidIndex(Idx); ++Idx) {}
-	for (; Elems.IsValidIndex(Idx + Num); ++Num) {}
-	return { Data + NumBytes(Idx, Stride), static_cast<uint64>(Num) };
-}
-
-inline void ReadSparseItems(FExistingItems& Dst, const FScriptSparseArray& Src, const FScriptSparseArrayLayout& Layout)
-{
-	const uint8* Data = static_cast<const uint8*>(Src.GetData(0, Layout));
-
-	if (FExistingItemSlice LastRead = Dst.Slice)
-	{
-		// Continue partial response
-		int64 TotalBytesRead = static_cast<const uint8*>(LastRead.Data) - Data;
-		check(TotalBytesRead % Layout.Size == 0);
-		int32 NextIdx = /* skip known invalid */ 1 + TotalBytesRead / Layout.Size;
-		check(NextIdx < Src.GetMaxIndex());
-		check(NextIdx > static_cast<int32>(LastRead.Num));
-		Dst.Slice = GetContiguousSlice(NextIdx, Src, Data, Layout.Size);
-	}
-	else if (Src.IsCompact())
-	{
-		int32 Num = Src.Num();
-		Dst.SetAll(Num ? Data : nullptr, Num);
-	}
-	else
-	{
-		// Start partial response
-		Dst.NumTotal = Src.Num();
-		Dst.Stride = Layout.Size;
-		Dst.Slice = GetContiguousSlice(0, Src, Data, Layout.Size);
-	}
-}
-
-// Helps save TSet/TMap deltas
-struct FSubSetIterator
-{
-	const FScriptSetLayout		Layout;
-	const FScriptSet&			Set;
-	const TBitArray<>&			Subset;
-	const int32					Max;
-	int32						Idx;
-
-	FSubSetIterator(const FScriptSetLayout& InLayout, const FScriptSet& InSet, const TBitArray<>& InSubset)
-	: Layout(InLayout)
-	, Set(InSet)
-	, Subset(InSubset)
-	, Max(InSet.GetMaxIndex())
-	, Idx(Max > 0 ? Subset.Find(true) : INDEX_NONE)
-	{}
-
-	explicit operator bool() const	{ return Idx != INDEX_NONE; }
-	const void* operator*() const	{ return Set.GetData(Idx, Layout); }
-	void operator++()				{ Idx = ++Idx < Max ? Subset.FindFrom(true, Idx) : INDEX_NONE; }
-	uint32 CountNum() const			{ return Subset.CountSetBits(); }
-};
-
-////////////////////////////////////////////////////////////////////////////////////////////////
-
 // Helpers to avoid using leaf FProperty instances after binding
 //
 // Below must match FFloatProperty, FDoubleProperty, FBoolProperty, FEnumProperty, TNumericProperty
@@ -575,20 +551,20 @@ template<typename T>
 inline bool		LeafPropertyIdentical(T A, T B) requires (std::is_unsigned_v<T>)	{ return A == B; }
 
 // Type-erased just enough to call LeafPropertyHash / LeafPropertyIdentical and FLeafRangeLoadView::As/AsBitcast
-enum class EPropertyKind : uint8 { Range, Struct, Bool, Integral8, Integral16, Integral32, Integral64, Float32, Float64 };
+enum class EPropertyKind : uint8 { Range, Struct, Bool, U8, U16, U32, U64, F32, F64 };
 
 template<EPropertyKind Kind>
-struct TLeafTypeOf											{ using Type = void; };
-template<> struct TLeafTypeOf<EPropertyKind::Bool>			{ using Type = bool; };
-template<> struct TLeafTypeOf<EPropertyKind::Integral8>		{ using Type = uint8; };
-template<> struct TLeafTypeOf<EPropertyKind::Integral16>	{ using Type = uint16; };
-template<> struct TLeafTypeOf<EPropertyKind::Integral32>	{ using Type = uint32; };
-template<> struct TLeafTypeOf<EPropertyKind::Integral64>	{ using Type = uint64; };
-template<> struct TLeafTypeOf<EPropertyKind::Float32>		{ using Type = float; };
-template<> struct TLeafTypeOf<EPropertyKind::Float64>		{ using Type = double; };
+struct TEquivalentLeafType									{ using Type = void; };
+template<> struct TEquivalentLeafType<EPropertyKind::Bool>	{ using Type = bool; };
+template<> struct TEquivalentLeafType<EPropertyKind::U8>	{ using Type = uint8; };
+template<> struct TEquivalentLeafType<EPropertyKind::U16>	{ using Type = uint16; };
+template<> struct TEquivalentLeafType<EPropertyKind::U32>	{ using Type = uint32; };
+template<> struct TEquivalentLeafType<EPropertyKind::U64>	{ using Type = uint64; };
+template<> struct TEquivalentLeafType<EPropertyKind::F32>	{ using Type = float; };
+template<> struct TEquivalentLeafType<EPropertyKind::F64>	{ using Type = double; };
 
 template<EPropertyKind Kind>
-using LeafTypeOf = typename TLeafTypeOf<Kind>::Type;
+using EquivalentLeafType = typename TEquivalentLeafType<Kind>::Type;
 
 template<typename LeafType>
 uint32 LeafHash(const void* In)
@@ -623,7 +599,7 @@ inline EPropertyKind GetPropertyKind(FLeafBindType In)
 
 	if (Type == ELeafType::Float)
 	{
-		return Width == ELeafWidth::B32 ? EPropertyKind::Float32 : EPropertyKind::Float64;
+		return Width == ELeafWidth::B32 ? EPropertyKind::F32 : EPropertyKind::F64;
 	}
 	else if (Type == ELeafType::Bool)
 	{
@@ -631,10 +607,10 @@ inline EPropertyKind GetPropertyKind(FLeafBindType In)
 	}
 	else switch (Width)
 	{
-		case ELeafWidth::B8:	return EPropertyKind::Integral8;
-		case ELeafWidth::B16:	return EPropertyKind::Integral16;
-		case ELeafWidth::B32:	return EPropertyKind::Integral32;
-		default:				return EPropertyKind::Integral64;
+		case ELeafWidth::B8:	return EPropertyKind::U8;
+		case ELeafWidth::B16:	return EPropertyKind::U16;
+		case ELeafWidth::B32:	return EPropertyKind::U32;
+		default:				return EPropertyKind::U64;
 	}								
 }
 
@@ -648,353 +624,547 @@ static EPropertyKind GetPropertyKind(FMemberBinding In)
 
 ////////////////////////////////////////////////////////////////////////////////////////////////
 
-inline uint32 HashRangeBindings(TConstArrayView<FRangeBinding> In)
-{
-	return static_cast<uint32>(FXxHash64::HashBuffer(In.GetData(), In.NumBytes()).Hash);
-}
-
-inline uint32 HashSkipOffset(FMemberBinding In)
+// Inner leaf property, e.g. FEnumProperty, FNumericProperty
+template<Arithmetic EquivalentType>
+struct TInnerLeafProperty
 {
-	uint32 Out = HashCombineFast(GetTypeHash(In.InnermostSchema), GetTypeHash(In.InnermostType));
-	return In.RangeBindings.IsEmpty() ? Out : HashCombineFast(Out, HashRangeBindings(In.RangeBindings));
-}
+	static constexpr EMemberKind	Kind = EMemberKind::Leaf;
+	static constexpr bool			bConstruct = false;
+	static constexpr bool			bDestruct = false;
+	static constexpr bool			bHashable = true;
+	static constexpr int32			Size = sizeof(EquivalentType);
+
+	TInnerLeafProperty(FProperty* In) { check(Size == In->GetElementSize());}
+
+	inline static void				InitItem(void* In)						{} // Note doesn't zero out items about to be overwritten
+	inline static void				DestroyItem(void* In)					{}
+	inline static EquivalentType	Cast(const void* In)					{ return *static_cast<const EquivalentType*>(In); } // Todo: Consider if BitCast is needed
+	inline static uint32			Hash(const void* In)					{ return LeafPropertyHash(Cast(In)); }
+	inline static bool				Identical(const void* A, const void* B) { return LeafPropertyIdentical(Cast(A), Cast(B)); }
+};
 
-inline bool EqSkipOffset(FMemberBinding A, FMemberBinding B)
+// Inner range-bound property, e.g. FArrayProperty, FStringProperty, FSetProperty
+struct FInnerRangeProperty
 {
-	return A.InnermostType == B.InnermostType && A.InnermostSchema == B.InnermostSchema && Algo::Compare(A.RangeBindings, B.RangeBindings);
-}
+	static constexpr EMemberKind	Kind = EMemberKind::Range;
+	static constexpr bool			bConstruct = false;
+	static constexpr bool			bDestruct = true;
 
-// Instead of TMap KeyFuncs
-struct FParameterBinding : FMemberBinding
-{
-	friend uint32 GetTypeHash(FParameterBinding In) { return HashSkipOffset(In); };
-	inline bool operator==(FParameterBinding O) const { return EqSkipOffset(*this, O); }
-};
+	FProperty*						Property;
+	uint32							Size;
+	bool							bHashable;
 
-////////////////////////////////////////////////////////////////////////////////////////////////
+	FInnerRangeProperty(FProperty* In)
+	: Property(In)
+	, Size(In->GetElementSize())
+	, bHashable(HasHash(In))
+	{
+		check(!HasConstructor(In));
+		check(HasDestructor(In));
+	}
 
-inline bool operator==(FScriptSetLayout A, FScriptSetLayout B)
-{
-	return FMemory::Memcmp(&A, &B, sizeof(A)) == 0;
-}
+	inline void						InitItem(void* In) const				{ FMemory::Memzero(In, Size); }
+	inline void						DestroyItem(void* In) const				{ DestroyValue(Property, In); }
+};
 
-// Used for leaf sets
-template<EPropertyKind ItemKind, typename LeafType>
-struct TSetPropertyBase
+// Inner struct-bound property, e.g. FStructProperty, FNameProperty, FObjectProperty
+struct FInnerStructProperty
 {
-	using SetElemT = TSetElement<LeafType>;
+	static constexpr EMemberKind	Kind = EMemberKind::Struct;
+
+	FProperty*						Property;
+	uint32							Size;
+	bool							bConstruct;
+	bool							bDestruct;
+	bool							bHashable;
+
+	FInnerStructProperty(FProperty* In)
+	: Property(In)
+	, Size(In->GetElementSize())
+	, bConstruct(HasConstructor(In))
+	, bDestruct(HasDestructor(In))
+	, bHashable(HasHash(In))
+	{}
 
-	static constexpr int32							ElementSize = sizeof(LeafType);
-	static constexpr uint16							MaxLoadScratchItems = sizeof(FLoadRangeContext::Scratch) / ElementSize;
-	static constexpr FScriptSetLayout				Layout = {offsetof(SetElemT, HashNextId), offsetof(SetElemT, HashIndex), sizeof(SetElemT), {alignof(SetElemT), sizeof(SetElemT)}};
-	static constexpr bool							bDestruct = false;
-	static constexpr bool							bConstruct = false;
-	static constexpr SIZE_T							BuiltItemSize = ElementSize;
-	
-	const FMemberBindType							InnerBindType;
-	const FMemberType								InnerType;
-	const FOptionalSchemaId							InnermostSchema;
-	const FRangeBinding								RangeBinding;					
+	inline void						InitItem(void* Item) const
+	{ 
+		if (bConstruct)
+		{
+			ConstructValue(Property, Item);
+		}
+		else
+		{
+			FMemory::Memzero(Item, Size);
+		}
+	}
 
-	TSetPropertyBase(FSetProperty* Set, FMemberBinding Elem, const IItemRangeBinding* This)
-	: InnerBindType(Elem.InnermostType)
-	, InnerType(ToLeafType(Elem.InnermostType.AsLeaf()))
-	, InnermostSchema(Elem.InnermostSchema)
-	, RangeBinding(*This, DefaultSizeType)
+	inline void						DestroyItem(void* Item) const
 	{
-		check(Layout == Set->SetLayout);
-		check(ElementSize == Set->ElementProp->GetElementSize());
-		check(ElementSize == SizeOf(InnerType.AsLeaf().Width));
+		if (bDestruct)
+		{
+			DestroyValue(Property, Item);
+		}
 	}
+};
 
-	static constexpr SIZE_T GetStride()								{ return sizeof(TSetElement<LeafType>); }
-	static constexpr void	InitCompactElems(void*, uint64, SIZE_T) {}
-	static constexpr void	DestructElems(FScriptSet& Set)			{}
+template<EPropertyKind Kind> struct TSelectInnerProperty		{ using Type = TInnerLeafProperty<EquivalentLeafType<Kind>>; };
+template<> struct TSelectInnerProperty<EPropertyKind::Range>	{ using Type = FInnerRangeProperty; };
+template<> struct TSelectInnerProperty<EPropertyKind::Struct>	{ using Type = FInnerStructProperty; };
+
+template<EPropertyKind Kind>
+using TInnerProperty = typename TSelectInnerProperty<Kind>::Type;
 
-	static constexpr int32 FindIndex(const FScriptSet& Set, const void* Elem)
+template<typename InnerPropertyType>
+inline auto	MakeHashFn(const InnerPropertyType& Inner)
+{
+	if constexpr (InnerPropertyType::Kind == EMemberKind::Leaf)
 	{
-		return Set.FindIndex(Elem, Layout, &LeafHash<LeafType>, &LeafIdentical<LeafType>);
+		return &InnerPropertyType::Hash;
 	}
-
-	static constexpr void Rehash(FScriptSet& Set)
+	else
 	{
-		Set.Rehash(Layout, &LeafHash<LeafType>);
+		return [P = Inner.Property](const void* In) { return HashValue(P, In); };
 	}
+}
 
-	inline FRangeMemberBinding MakeRangeMemberBinding() const
+template<typename InnerPropertyType>
+inline auto	MakeIdenticalFn(const InnerPropertyType& Inner)
+{
+	if constexpr (InnerPropertyType::Kind == EMemberKind::Leaf)
+	{
+		return &InnerPropertyType::Identical;
+	}
+	else
 	{
-		return { &InnerBindType, &RangeBinding, 1, InnermostSchema, 0 };
+		return [P = Inner.Property](const void* A, const void* B) { return P->Identical(A, B); };
 	}
+}
 
-	inline FMemberSchema MakeRangeSchema() const
+template<typename InnerPropertyType>
+inline void InitStridedItems(const InnerPropertyType& Inner, void* Items, uint64 Num, SIZE_T Stride)
+{
+	if constexpr (InnerPropertyType::Kind == EMemberKind::Leaf)
+	{}
+	else if (Inner.bConstruct)
+	{
+		ConstructValues(Inner.Property, static_cast<uint8*>(Items), Num, Stride);
+	}
+	else
 	{
-		return { FMemberType(DefaultSizeType), InnerType, 1, InnermostSchema, nullptr };
+		MemzeroStrided(static_cast<uint8*>(Items), Num, Inner.Size, Stride);
 	}
+}
 
-	inline static void SaveDeltaItems(FAnyRangeSaver& Out, FSubSetIterator It, const FSaveContext&)
+template<typename InnerPropertyType>
+inline void DestroyStridedItems(const InnerPropertyType& Inner, void* Items, uint64 Num, SIZE_T Stride)
+{
+	if constexpr (InnerPropertyType::Kind == EMemberKind::Leaf)
+	{}
+	else if (Inner.bDestruct)
 	{
-		// Can be optimized by checking IsCompact() and adding FAnyRangeSaver::AddLeaves
-		for (; It; ++It)
-		{
-			Out.AddLeaf(*It, sizeof(LeafType));
-		}
+		DestroyValues(Inner.Property, static_cast<uint8*>(Items), Num, Stride);
 	}
-};
+}
+
+////////////////////////////////////////////////////////////////////////////////////////////////
 
-// Used for non-leaf sets
-template<EPropertyKind ItemKind>
-struct TSetPropertyBase<ItemKind, void>
+template<typename EquivalentType>
+struct TLeafRangeSerializer
 {
-	static constexpr SIZE_T							BuiltItemSize = sizeof(FBuiltValue);
+	using RangeSaver = TLeafRangeSaver<EquivalentType>;
+	static constexpr SIZE_T	Size = sizeof(EquivalentType);
 
-	FProperty const* const							ElemProp;
-	const FScriptSetLayout							Layout;
-	const bool										bDestruct;
-	const bool										bConstruct;
-	const uint16									MaxLoadScratchItems;
-	const int32										ElementSize;
-	const FOptionalSchemaId							InnermostSchema;
-	TArray<FMemberType, TInlineAllocator<8>>		InnerTypes;
-	TArray<FMemberBindType, TInlineAllocator<8>>	InnerBindTypes;
-	TArray<FRangeBinding, TInlineAllocator<2>>		RangeBindings;
+	FMemberType				InnerType;
+	FOptionalSchemaId		EnumId;
 
-	TSetPropertyBase(FSetProperty* In, FMemberBinding Elem, const IItemRangeBinding* This)
-	: ElemProp(In->ElementProp)
-	, Layout(In->SetLayout)
-	, bDestruct(HasDestructor(In->ElementProp))
-	, bConstruct(HasConstructor(In->ElementProp))
-	, MaxLoadScratchItems(IntCastChecked<uint16>(sizeof(FLoadRangeContext::Scratch) / In->ElementProp->GetElementSize()))
-	, ElementSize(In->ElementProp->GetElementSize())
-	, InnermostSchema(Elem.InnermostSchema)
+	explicit TLeafRangeSerializer(FMemberBinding In)
+	: InnerType(ToLeafType(In.InnermostType.AsLeaf()))
+	, EnumId(In.InnermostSchema)
 	{
-		for (FRangeBinding Inner : Elem.RangeBindings)
-		{
-			InnerTypes.Emplace(Inner.GetSizeType());
-			InnerBindTypes.Emplace(Inner.GetSizeType());
-		}
-		InnerTypes.Emplace(Elem.InnermostType.IsStruct() ? FMemberType(Elem.InnermostType.AsStruct()) 
-														 : FMemberType(ToLeafType(Elem.InnermostType.AsLeaf())));
-		InnerBindTypes.Emplace(Elem.InnermostType);
-		RangeBindings.Emplace(*This, DefaultSizeType);
-		RangeBindings.Append(Elem.RangeBindings);
-		check(RangeBindings.Num() <= 0xFFFF);
-		check(Elem.RangeBindings.IsEmpty() == (ItemKind != EPropertyKind::Range));
-		check(ElemProp->GetElementSize() == ElemSize());
+		check(In.RangeBindings.IsEmpty());
+		check(InnerType.AsLeaf().Width == WidthOf(Size));
 	}
 
-	inline SIZE_T			GetStride() const					{ return Layout.Size; }
-	inline SIZE_T			ElemSize() const					{ return ElementSize; }
-	inline auto				MakeHashFn() const					{ return [P = ElemProp](const void* In) { return P->GetValueTypeHash(In); }; }
-	inline auto				MakeIdenticalFn() const				{ return [P = ElemProp](const void* A, const void* B) { return P->Identical(A, B); }; }
+	inline static EquivalentType Cast(const void* In)
+	{
+		return *static_cast<const EquivalentType*>(In);
+	}
 
-	inline int32 FindIndex(const FScriptSet& Set, const void* Elem) const
+	inline FMemberSchema MakeMemberSchema() const
 	{
-		return Set.FindIndex(Elem, Layout, MakeHashFn(), MakeIdenticalFn());
+		return { FMemberType(DefaultSizeType), InnerType, 1, EnumId, nullptr };
 	}
 
-	inline void Rehash(FScriptSet& Set) const
+	inline static EquivalentType SaveItem(const void* In, const FSaveContext&)
 	{
-		Set.Rehash(Layout, MakeHashFn());
+		return Cast(In);
 	}
 
-	inline FRangeMemberBinding MakeRangeMemberBinding() const
+	inline static void LoadItem(void* Dst, FByteReader& SrcBytes, FBitCacheReader& SrcBits, const FLoadBatch&) requires (std::is_same_v<EquivalentType, bool>)
 	{
-		return { InnerBindTypes.GetData(), RangeBindings.GetData(), static_cast<uint16>(InnerBindTypes.Num()), InnermostSchema, 0 };
+		*static_cast<bool*>(Dst) = SrcBits.GrabNext(SrcBytes);
 	}
 
-	inline FMemberSchema MakeRangeSchema() const
+	inline static void LoadItem(void* Dst, FByteReader& SrcBytes, FBitCacheReader&, const FLoadBatch&)
 	{
-		return { FMemberType(DefaultSizeType), InnerTypes[0], static_cast<uint16>(InnerTypes.Num()), InnermostSchema, InnerTypes.Num() > 1 ? InnerTypes.GetData() : nullptr };
+		*static_cast<EquivalentType*>(Dst) = Cast(SrcBytes.GrabBytes(Size));
 	}
+};
+
+struct FStructRangeSerializer
+{
+	using RangeSaver = FStructRangeSaver;
+	
+	FMemberType				InnerType;
+	FStructSchemaId			Id;
 
-	inline TConstArrayView<FRangeBinding> GetInnerBindings() const
+	explicit FStructRangeSerializer(FMemberBinding Item)
+	: InnerType(Item.InnermostType.AsStruct())
+	, Id(static_cast<FStructSchemaId>(Item.InnermostSchema.Get()))
 	{
-		return { RangeBindings.GetData() + 1, RangeBindings.Num() - 1 };
+		check(Item.RangeBindings.IsEmpty());
 	}
 
-	inline void* ConstructElem(void* Item) const
+	inline FMemberSchema MakeMemberSchema() const
 	{
-		if (bConstruct)
-		{
-			InitializeValue(ElemProp, Item);
-		}
-		else
-		{
-			FMemory::Memzero(Item, ElemSize());
-		}
-		return Item;
+		return { FMemberType(DefaultSizeType), InnerType, 1, FOptionalSchemaId(Id), nullptr };
 	}
 
-	inline void InitCompactElems(void* Items, uint64 Num, SIZE_T Stride) const
+	inline FBuiltStruct* SaveItem(const void* In, const FSaveContext& Ctx) const
 	{
-		if (bConstruct)
-		{
-			InitializeValues(ElemProp, static_cast<uint8*>(Items), Num, Stride);
-		}
-		else
-		{
-			for (uint8* It = static_cast<uint8*>(Items), *End = It + Num * Stride; It != End; It += Stride)
-			{
-				FMemory::Memzero(It, ElemSize());
-			}
-		}
+		return SaveStruct(In, Id, Ctx);
 	}
 
-	inline void DestructElem(void* Elem) const
+	inline void LoadItem(void* Dst, FByteReader& SrcBytes, FBitCacheReader&, const FLoadBatch& Batch) const
 	{
-		if (bDestruct)
-		{
-			DestroyValue(ElemProp, Elem);
-		}
+		LoadStruct(Dst, FByteReader(SrcBytes.GrabSkippableSlice()), Id, Batch);
 	}
+};
 
-	inline void DestructElems(FScriptSet& Set) const
+struct FNestedRangeSerializer
+{
+	using RangeSaver = FNestedRangeSaver;
+
+	FOptionalSchemaId								InnermostId;
+	uint16											NumInners;
+	TArray<FMemberType, TInlineAllocator<8>>		InnerTypes;
+	TArray<FMemberBindType, TInlineAllocator<8>>	InnerBindTypes;
+	TArray<FRangeBinding, TInlineAllocator<2>>		InnerBindings;
+
+	explicit FNestedRangeSerializer(FMemberBinding Item)
+	: InnermostId(Item.InnermostSchema)
+	, NumInners(IntCastChecked<uint16>(1 + Item.RangeBindings.Num()))
+	, InnerBindings(Item.RangeBindings)
 	{
-		check(bDestruct);
-		uint8* It = static_cast<uint8*>(Set.GetData(0, Layout));
-		SIZE_T Stride = GetStride();
-		if (IsCompact(Set))
+		check(NumInners >= 2);
+		for (FRangeBinding Inner : Item.RangeBindings)
 		{
-			DestroyValues(ElemProp, It, Set.Num(), Stride);
-		}
-		else
-		{
-			for (int32 Idx = 0, Max = Set.GetMaxIndex(); Idx < Max; ++Idx)
-			{
-				if (Set.IsValidIndex(Idx))
-				{
-					DestroyValue(ElemProp, It);
-				}
-				It += Stride;
-			}
+			InnerTypes.Emplace(Inner.GetSizeType());
+			InnerBindTypes.Emplace(Inner.GetSizeType());
 		}
+		InnerTypes.Emplace(Item.InnermostType.IsStruct() ? FMemberType(Item.InnermostType.AsStruct()) 
+														 : FMemberType(ToLeafType(Item.InnermostType.AsLeaf())));
+		InnerBindTypes.Emplace(Item.InnermostType);
 	}
 
-	inline void SaveDeltaItems(FAnyRangeSaver& Out, FSubSetIterator& It, const FSaveContext& Ctx) const requires (ItemKind == EPropertyKind::Range)
+	inline FMemberSchema MakeMemberSchema() const
 	{
-		FRangeMemberBinding Inner = GetInnerRange(MakeRangeMemberBinding());
-		for (; It; ++It)
-		{
-			Out.AddRange(SaveRange(*It, Inner, Ctx));
-		}
+		return { FMemberType(DefaultSizeType), InnerTypes[0], NumInners, InnermostId, InnerTypes.GetData() };
 	}
 
-	inline void SaveDeltaItems(FAnyRangeSaver& Out, FSubSetIterator& It, const FSaveContext& Ctx) const requires (ItemKind == EPropertyKind::Struct)
+	FBuiltRange* SaveItem(const void* In, const FSaveContext& Ctx) const
 	{
-		FStructSchemaId Id = static_cast<FStructSchemaId>(InnermostSchema.Get());
-		// Todo: Lookup Ctx.Defaults and call SaveStructDelta instead
-		for (; It; ++It)
-		{
-			Out.AddStruct(SaveStruct(*It, Id, Ctx));
-		}
+		FRangeMemberBinding Member = { InnerBindTypes.GetData() + 1, InnerBindings.GetData(), NumInners - 1, InnermostId, 0 };
+		return SaveRange(In, Member, Ctx);
+	}
+
+	inline void LoadItem(void* Dst, FByteReader& SrcBytes, FBitCacheReader& SrcBits, const FLoadBatch& Batch) const
+	{
+		FRangeLoadSchema Schema = {InnerTypes[0], InnermostId, InnerTypes.GetData() + 1, Batch};
+		LoadRange(Dst, SrcBytes, SrcBits, DefaultSizeType, Schema, InnerBindings);
 	}
 };
 
-template<EPropertyKind ItemKind>
-struct TSetPropertyBinding : IItemRangeBinding, ICustomBinding, TSetPropertyBase<ItemKind, LeafTypeOf<ItemKind>>
+template<EPropertyKind Kind> struct TSelectRangeSerializer		{ using Type = TLeafRangeSerializer<EquivalentLeafType<Kind>>; };
+template<> struct TSelectRangeSerializer<EPropertyKind::Range>	{ using Type = FNestedRangeSerializer; };
+template<> struct TSelectRangeSerializer<EPropertyKind::Struct>	{ using Type = FStructRangeSerializer; };
+
+template<EPropertyKind Kind>
+using TPropertyRangeSerializer = typename TSelectRangeSerializer<Kind>::Type;
+
+////////////////////////////////////////////////////////////////////////////////////////////////
+
+inline FScriptSparseArray& AsSparseArray(FScriptSet& In) { return reinterpret_cast<FScriptSparseArray&>(In); }
+inline FScriptSparseArray& AsSparseArray(FScriptMap& In) { return reinterpret_cast<FScriptSparseArray&>(In); }
+
+template<typename ScriptSet>
+inline bool IsCompact(const ScriptSet& Set)
 {
-	using LeafType = LeafTypeOf<ItemKind>;
-	using Base = TSetPropertyBase<ItemKind, LeafType>;
-	static constexpr bool bLeafItems = !std::is_void_v<LeafType>;
-	
-	TSetPropertyBinding(FSetProperty* In, FMemberBinding Elem) : Base(In, Elem, this) {}
-	
-	using Base::ElementSize;
-	using Base::MaxLoadScratchItems;
-	using Base::Layout;
-	using Base::bConstruct;
-	using Base::bDestruct;
-	using Base::GetStride;
+	return Set.NumUnchecked() == Set.GetMaxIndex();
+}
 
-	virtual void ReadItems(FSaveRangeContext& Ctx) const override
+// There's no TScriptSparseArray::SetNumUninitialized() (yet), 
+// reserve using Empty() and add items one by one instead
+template<class ScriptType, class LayoutType>
+uint8* SetNumUninitialized(ScriptType& Dst, const LayoutType& Layout, uint64 Num)
+{
+	check(Dst.IsEmpty());
+	Dst.Empty(static_cast<int32>(Num), Layout);
+	for (uint64 Idx = 0; Idx < Num; ++Idx)
 	{
-		ReadSparseItems(/* out */ Ctx.Items, Ctx.Request.GetRange<FScriptSparseArray>(), Layout.SparseArrayLayout);
+		Dst.AddUninitialized(Layout);
 	}
+	check(IsCompact(Dst));
 
-	virtual void MakeItems(FLoadRangeContext& Ctx) const override
-	{		
-		// Load into scratch and memcpy to sparse TSetElement array to avoid FProperty::CopyValuesInternal dependency
-		// TSet elements are always trivially relocatable. Scratch items are constructed but not destroyed.
-		const bool bScratch = Ctx.Request.NumTotal() > 1 && MaxLoadScratchItems > 1;
-		const int32 NumTotal = static_cast<int32>(Ctx.Request.NumTotal());
-		const SIZE_T ItemSize = ElementSize;
-		const SIZE_T Stride = GetStride();
+	return static_cast<uint8*>(Dst.GetData(0, Layout));
+}
 
-		FScriptSet& Set = Ctx.Request.GetRange<FScriptSet>();
-		if (Ctx.Request.IsFirstCall())
-		{
-			if (bDestruct && !Set.IsEmpty())
-			{
-				Base::DestructElems(Set);
-			}
+inline FExistingItemSlice GetContiguousSlice(int32 Idx, const FScriptSparseArray& Elems, const uint8* Data, SIZE_T Stride)
+{
+	int32 Num = 1;
+	for (;!Elems.IsValidIndex(Idx); ++Idx) {}
+	for (; Elems.IsValidIndex(Idx + Num); ++Num) {}
+	return { Data + NumBytes(Idx, Stride), static_cast<uint64>(Num) };
+}
 
-			Set.Empty(NumTotal, Layout);
-			
-			if (NumTotal > 0)
-			{
-				// Todo: Add TScriptSparseArray::AppendUninitialized or SetNumUninitialized or Reset
-				for (int32 Idx = 0; Idx < NumTotal; ++Idx)
-				{
-					Set.AddUninitialized(Layout);
-				}
-				check(IsCompact(Set));
+// Save flat TSet/TMap
+inline void ReadSparseItems(FExistingItems& Dst, const FScriptSparseArray& Src, const FScriptSparseArrayLayout& Layout)
+{
+	const uint8* Data = static_cast<const uint8*>(Src.GetData(0, Layout));
 
-				if (bScratch)
-				{
-					int32 NumScratch = FMath::Min(int32(MaxLoadScratchItems), NumTotal);
-					Base::InitCompactElems(Ctx.Scratch, NumScratch, ItemSize);
-					Ctx.Items.Set(Ctx.Scratch, NumScratch, ItemSize);
-				}
-				else // Initialize all elements and then load them one by one
-				{
-					uint8* Data = static_cast<uint8*>(Set.GetData(0, Layout));
-					Base::InitCompactElems(Data, NumTotal, Stride);
-					Ctx.Items.Set(Data, 1, ItemSize);
-				}
+	if (FExistingItemSlice LastRead = Dst.Slice)
+	{
+		// Continue partial response
+		int64 TotalBytesRead = static_cast<const uint8*>(LastRead.Data) - Data;
+		check(TotalBytesRead % Layout.Size == 0);
+		int32 NextIdx = /* skip known invalid */ 1 + TotalBytesRead / Layout.Size;
+		check(NextIdx < Src.GetMaxIndex());
+		check(NextIdx > static_cast<int32>(LastRead.Num));
+		Dst.Slice = GetContiguousSlice(NextIdx, Src, Data, Layout.Size);
+	}
+	else if (Src.IsCompact())
+	{
+		int32 Num = Src.Num();
+		Dst.SetAll(Num ? Data : nullptr, Num);
+	}
+	else
+	{
+		// Start partial response
+		Dst.NumTotal = Src.Num();
+		Dst.Stride = Layout.Size;
+		Dst.Slice = GetContiguousSlice(0, Src, Data, Layout.Size);
+	}
+}
 
-				Ctx.Items.RequestFinalCall();
-			}
+////////////////////////////////////////////////////////////////////////////////////////////////
+
+// Helps save TSet/TMap deltas
+template<class ScriptType, class LayoutType>
+struct TSubSetIterator
+{
+	const LayoutType			Layout;
+	const ScriptType&			Set;
+	const TBitArray<>&			Subset;
+	const int32					Max;
+	int32						Idx;
+
+	TSubSetIterator(const LayoutType& InLayout, const ScriptType& InSet, const TBitArray<>& InSubset)
+	: Layout(InLayout)
+	, Set(InSet)
+	, Subset(InSubset)
+	, Max(InSet.GetMaxIndex())
+	, Idx(Max > 0 ? Subset.Find(true) : INDEX_NONE)
+	{}
+
+	explicit operator bool() const	{ return Idx != INDEX_NONE; }
+	const void* operator*() const	{ return Set.GetData(Idx, Layout); }
+	void operator++()				{ Idx = ++Idx < Max ? Subset.FindFrom(true, Idx) : INDEX_NONE; }
+	uint32 CountNum() const			{ return Subset.CountSetBits(); }
+};
+
+using FSubSetIterator = TSubSetIterator<FScriptSet, FScriptSetLayout>;
+using FSubMapIterator = TSubSetIterator<FScriptMap, FScriptMapLayout>;
+
+// Helps save TSet/TMap deltas
+template<class ScriptType, class LayoutType, class SerializerType>
+FTypedRange SaveAll(const ScriptType& Set, const LayoutType& Layout, const SerializerType& Serializer, const FSaveContext& Ctx)
+{
+	check(!Set.IsEmpty());
+
+	typename SerializerType::RangeSaver Range(Ctx.Scratch, static_cast<uint64>(Set.Num()));
+	for (int32 Idx = 0, Max = Set.GetMaxIndex(); Idx < Max; ++Idx)
+	{
+		if (Set.IsValidIndex(Idx))
+		{
+			Range.AddItem(Serializer.SaveItem(Set.GetData(Idx, Layout), Ctx));
 		}
-		else if (Ctx.Request.IsFinalCall())
+	}
+	return Range.Finalize(Serializer.MakeMemberSchema());
+}
+
+// Helps save TSet/TMap deltas
+template<class SubSetIteratorType, class SerializerType>
+FTypedRange SaveSome(SubSetIteratorType& It, const SerializerType& Serializer, const FSaveContext& Ctx)
+{
+	typename SerializerType::RangeSaver Range(Ctx.Scratch, It.CountNum());
+	for (; It; ++It)
+	{
+		Range.AddItem(Serializer.SaveItem(*It, Ctx));
+	}
+	return Range.Finalize(Serializer.MakeMemberSchema());
+}
+
+template<typename BindingType, typename ScriptType>
+void LoadSetDelta(const BindingType& Binding, ScriptType& Dst, FStructLoadView Src, ECustomLoadMethod Method)
+{
+	check(Method == ECustomLoadMethod::Assign);
+	FMemberLoader Members(Src);				
+	while (Members.HasMore())
+	{
+		if (Members.PeekName() == GUE.CommonIds.Add)
 		{
-			Base::Rehash(Set);
+			if (Dst.IsEmpty())
+			{
+				Binding.Append(Dst, Members.GrabRange());
+			}
+			else
+			{
+				Binding.Insert(Dst, Members.GrabRange());
+			}
+			check(!Members.HasMore());
+			break;
 		}
-		else
+
+		check(Members.PeekName() == GUE.CommonIds.Del);
+		FRangeLoadView Keys = Members.GrabRange();
+		if (!Dst.IsEmpty())
 		{
-			uint8* DstIt = static_cast<uint8*>(Set.GetData(static_cast<int32>(Ctx.Request.GetIndex()), Layout));
+			Binding.Remove(Dst, Keys);
+		}
+	}	
+}
 
-			if (bScratch)
-			{
-				// Move loaded items from scratch to sparse array
-				for (const uint8* SrcIt = Ctx.Items.GetData(), *SrcEnd = Ctx.Items.GetDataEnd(); SrcIt != SrcEnd; SrcIt += ItemSize)
-				{
-					FMemory::Memcpy(DstIt, SrcIt, ItemSize);
-					DstIt += Stride;
-				}
+template<typename BindingType, typename ScriptType>
+inline bool DiffSet(const BindingType& Binding, const ScriptType& A, const ScriptType& B)
+{
+	if (A.NumUnchecked() != B.NumUnchecked())
+	{
+		return true;
+	}
 
-				Ctx.Items.UpdateNum(FMath::Min(Ctx.Request.NumMore(), uint64(MaxLoadScratchItems)));
-			}
-			else // Load one by one
+	if (A.NumUnchecked() > 0)
+	{
+		for (int32 IdxA = 0, MaxA = A.GetMaxIndex(); IdxA < MaxA; ++IdxA)
+		{
+			if (A.IsValidIndex(IdxA) && !Binding.HasElem(B, A.GetData(IdxA, Binding.Layout)))
 			{
-				Ctx.Items.Set(DstIt, 1, ItemSize);
+				return true;
 			}
 		}
 	}
-	
-	inline bool HasElem(const FScriptSet& Set, const void* Elem) const
+
+	return false;
+}
+
+////////////////////////////////////////////////////////////////////////////////////////////////
+
+template<EPropertyKind ElemKind>
+struct TSetPropertyBinding : IItemRangeBinding, ICustomBinding
+{
+	using SetType = FScriptSet;
+	using LeafType = EquivalentLeafType<ElemKind>;
+	static constexpr bool bLeaves = !std::is_void_v<LeafType>;
+
+	const FScriptSetLayout							Layout;
+	const TInnerProperty<ElemKind>					Inner;
+	const TPropertyRangeSerializer<ElemKind>		Range;
+
+	TSetPropertyBinding(FSetProperty* In, FMemberBinding Elem)
+	: Layout(In->SetLayout)
+	, Inner(In->ElementProp)
+	, Range(Elem)
 	{
-		return Base::FindIndex(Set, Elem) != INDEX_NONE;
+		check(Layout.Size == GetStride());
+		check(Inner.bHashable);
+	}
+
+	inline SIZE_T GetStride() const requires (bLeaves)	{ return sizeof(TSetElement<LeafType>); }
+	inline SIZE_T GetStride() const						{ return Layout.Size; }
+
+	inline int32 FindIndex(const FScriptSet& Set, const void* Elem) const
+	{
+		return Set.FindIndex(Elem, Layout, MakeHashFn(Inner), MakeIdenticalFn(Inner));
 	}
 
 	inline void RemoveElem(FScriptSet& Set, const void* Elem) const
 	{
-		int32 Idx = Base::FindIndex(Set, Elem);
-		if (Idx != INDEX_NONE)
+		if (int32 Idx = FindIndex(Set, Elem); Idx != INDEX_NONE)
 		{
 			Set.RemoveAt(Idx, Layout);
 		}
 	}
 
+	inline bool HasElem(const FScriptSet& Set, const void* Elem) const
+	{
+		return FindIndex(Set, Elem) != INDEX_NONE;
+	}
+
+	inline void Rehash(FScriptSet& Set) const
+	{
+		Set.Rehash(Layout, MakeHashFn(Inner));
+	}
+	
+	inline void DestroyAll(FScriptSet& Set) const requires (bLeaves) {}
+	inline void DestroyAll(FScriptSet& Set) const
+	{
+		uint8* It = static_cast<uint8*>(Set.GetData(0, Layout));
+		SIZE_T Stride = GetStride();
+		if (IsCompact(Set))
+		{
+			DestroyStridedItems(Inner, It, Set.NumUnchecked(), Stride);
+		}
+		else
+		{
+			for (int32 Idx = 0, Max = Set.GetMaxIndex(); Idx < Max; ++Idx)
+			{
+				if (Set.IsValidIndex(Idx))
+				{
+					DestroyValue(Inner.Property, It);
+				}
+				It += Stride;
+			}
+		}
+	}
+	
+	virtual void ReadItems(FSaveRangeContext& Ctx) const override
+	{
+		ReadSparseItems(/* out */ Ctx.Items, Ctx.Request.GetRange<FScriptSparseArray>(), Layout.SparseArrayLayout);
+	}
+
+	virtual void MakeItems(FLoadRangeContext& Ctx) const override
+	{
+		FScriptSet& Set = Ctx.Request.GetRange<FScriptSet>();
+		int32 NewNum = static_cast<int32>(Ctx.Request.NumTotal());
+
+		if (Ctx.Request.IsFirstCall())
+		{
+			DestroyAll(Set);
+			Set.Empty(NewNum, Layout);
+			if (NewNum)
+			{
+				uint8* Items = SetNumUninitialized(Set, Layout, NewNum);
+				InitStridedItems(Inner, Items, NewNum, GetStride());
+				Ctx.Items.Set(Items, Ctx.Request.NumTotal(), GetStride());
+				Ctx.Items.RequestFinalCall();
+			}
+		}
+		else
+		{
+			check(Ctx.Request.IsFinalCall());
+			Rehash(Set);
+		}
+	}
+	
 	virtual void SaveCustom(FMemberBuilder& Dst, const void* InSrc, const void* InDefault, const FSaveContext& Ctx) override
 	{
 		const FScriptSet& Src = *static_cast<const FScriptSet*>(InSrc);
@@ -1004,8 +1174,7 @@ struct TSetPropertyBinding : IItemRangeBinding, ICustomBinding, TSetPropertyBase
 		{
 			if (Default && !Default->IsEmpty())
 			{
-				FBuiltRange* Range = SaveRange(Default, Base::MakeRangeMemberBinding(), Ctx);
-				Dst.AddRange(GUE.CommonIds.Del, { Base::MakeRangeSchema(), Range });
+				Dst.AddRange(GUE.CommonIds.Del, SaveAll(*Default, Layout, Range, Ctx));
 			}
 		}
 		else if (Default && !Default->IsEmpty())
@@ -1017,7 +1186,7 @@ struct TSetPropertyBinding : IItemRangeBinding, ICustomBinding, TSetPropertyBase
 			}
 			if (FSubSetIterator Deleted{Layout, *Default, DeletedIds})
 			{
-				Dst.AddRange(GUE.CommonIds.Del, SaveDeltaRange(Deleted, Ctx));
+				Dst.AddRange(GUE.CommonIds.Del, SaveSome(Deleted, Range, Ctx));
 			}
 			
 			TBitArray<> AddedIds(false, Src.GetMaxIndex());
@@ -1027,178 +1196,127 @@ struct TSetPropertyBinding : IItemRangeBinding, ICustomBinding, TSetPropertyBase
 			}
 			if (FSubSetIterator Added{Layout, Src, AddedIds})
 			{
-				Dst.AddRange(GUE.CommonIds.Add, SaveDeltaRange(Added, Ctx));
+				Dst.AddRange(GUE.CommonIds.Add, SaveSome(Added, Range, Ctx));
 			}
 		}
 		else
 		{
-			FBuiltRange* Range = SaveRange(&Src, Base::MakeRangeMemberBinding(), Ctx);
-			Dst.AddRange(GUE.CommonIds.Add, { Base::MakeRangeSchema(), Range });
+			Dst.AddRange(GUE.CommonIds.Add, SaveAll(Src, Layout, Range, Ctx));
 		}
 	}
 
-	virtual void LoadCustom(void* InDst, FStructLoadView Src, ECustomLoadMethod Method) const override
+	virtual void LoadCustom(void* Dst, FStructLoadView Src, ECustomLoadMethod Method) const override
 	{
-		check(Method == ECustomLoadMethod::Assign);
-		FScriptSet& Dst = *static_cast<FScriptSet*>(InDst);
-		FMemberLoader Members(Src);				
-		while (Members.HasMore())
-		{
-			if (Members.PeekName() == GUE.CommonIds.Add)
-			{
-				if (Dst.IsEmpty())
-				{
-					Append(Dst, Members.GrabRange());
-				}
-				else
-				{
-					Insert(Dst, Members.GrabRange());
-				}
-				check(!Members.HasMore());
-				break;
-			}
-
-			check(Members.PeekName() == GUE.CommonIds.Del);
-			if (!Dst.IsEmpty())
-			{
-				Remove(Dst, Members.GrabRange());
-			}
-		}	
+		LoadSetDelta(*this, *static_cast<FScriptSet*>(Dst), Src, Method);
 	}
 
-	virtual bool DiffCustom(const void* InA, const void* InB) const override
+	virtual bool DiffCustom(const void* A, const void* B) const override
 	{
-		const FScriptSet& A = *static_cast<const FScriptSet*>(InA);
-		const FScriptSet& B = *static_cast<const FScriptSet*>(InB);
-
-		if (A.NumUnchecked() != B.NumUnchecked())
-		{
-			return true;
-		}
-
-		if (A.NumUnchecked() > 0)
-		{
-			for (int32 IdxA = 0, MaxA = A.GetMaxIndex(); IdxA < MaxA; ++IdxA)
-			{
-				if (A.IsValidIndex(IdxA) && !HasElem(B, A.GetData(IdxA, Layout)))
-				{
-					return true;
-				}
-			}
-		}
-
-		return false;
+		return DiffSet(*this, *static_cast<const FScriptSet*>(A), *static_cast<const FScriptSet*>(B));
 	}
-	
+
 	// Load into empty set
 	inline void Append(FScriptSet& Dst, FRangeLoadView Src) const
 	{
-		// There's no TScriptSparseArray::SetNumUninitialized() (yet), 
-		// reserve using Empty() and add items one by one instead
-		Dst.Empty(static_cast<int32>(Src.Num()), Layout);
-		for (uint64 Idx = 0; Idx < Src.Num(); ++Idx)
-		{
-			Dst.AddUninitialized(Layout);
-		}
-		check(IsCompact(Dst));
-		uint8* It = static_cast<uint8*>(Dst.GetData(0, Layout));
 		SIZE_T Stride = GetStride();
-		Base::InitCompactElems(It, Src.Num(), Stride);
+		uint8* It = SetNumUninitialized(Dst, Layout, Src.Num());
+		InitStridedItems(Inner, It, Src.Num(), Stride);
 
-		if constexpr (bLeafItems)
+		if constexpr (bLeaves)
 		{
-			for (LeafType Leaf : CastAs<LeafType>(Src.AsLeaves()))
+			for (LeafType Item : CastAs<LeafType>(Src.AsLeaves()))
 			{
-				*reinterpret_cast<LeafType*>(It) = Leaf;
+				*reinterpret_cast<LeafType*>(It) = Item;
 				It += Stride;
 			}
 		}
-		else if constexpr (ItemKind == EPropertyKind::Range)
+		else if constexpr (ElemKind == EPropertyKind::Range)
 		{
-			TConstArrayView<FRangeBinding> Inners = Base::GetInnerBindings();
-			for (FRangeLoadView Range : Src.AsRanges())
+			TConstArrayView<FRangeBinding> InnerBindings = Range.InnerBindings;
+			for (FRangeLoadView Item : Src.AsRanges())
 			{
-				LoadRange(It, Range, Inners);
+				LoadRange(It, Item, InnerBindings);
 				It += Stride;
 			}
 		}
 		else
 		{
-			for (FStructLoadView Struct : Src.AsStructs())
+			for (FStructLoadView Item : Src.AsStructs())
 			{
-				LoadStruct(It, Struct);
+				LoadStruct(It, Item);
 				It += Stride;
 			}
 		}
 
-		Base::Rehash(Dst);
+		Rehash(Dst);
 	}
 
-	
-	inline void Insert(FScriptSet& Dst, FRangeLoadView Src) const requires(bLeafItems)
+	// Load leaves into non-empty set
+	inline void Insert(FScriptSet& Dst, FRangeLoadView Src) const requires(bLeaves)
 	{	
-		for (LeafType Leaf : CastAs<LeafType>(Src.AsLeaves()))
+		for (LeafType Item : CastAs<LeafType>(Src.AsLeaves()))
 		{
-			if (!HasElem(Dst, &Leaf))
+			if (!HasElem(Dst, &Item))
 			{
 				 void* Elem = Dst.GetData(Dst.AddUninitialized(Layout), Layout);
-				 *static_cast<LeafType*>(Elem) = Leaf;
+				 *static_cast<LeafType*>(Elem) = Item;
 			}
 		}
 
-		Base::Rehash(Dst);
+		Rehash(Dst);
 	}
 
-	inline void* AddConstructed(FScriptSparseArray& Dst, int32& OutIdx) const
+	inline void* AddItem(FScriptSparseArray& Dst, int32& OutIdx) const
 	{
 		OutIdx = Dst.AddUninitialized(Layout.SparseArrayLayout);
-		return Base::ConstructElem(Dst.GetData(OutIdx, Layout.SparseArrayLayout));
+		void* Out = Dst.GetData(OutIdx, Layout.SparseArrayLayout);
+		Inner.InitItem(Out);
+		return Out;
 	}
 
 	// Load structs or ranges into non-empty set
-	inline void Insert(FScriptSet& Dst, FRangeLoadView Src) const
+	inline void Insert(FScriptSet& DstSet, FRangeLoadView Src) const
 	{
 		// Written to avoid FProperty::CopyCompleteValue_InContainer dependency
 		// Items are loaded directly into sparse array and then removed if a duplicate existed
-
-		FScriptSparseArray& DstItems = reinterpret_cast<FScriptSparseArray&>(Dst);
-		const int32 OldNum = DstItems.NumUnchecked();
+		FScriptSparseArray& Dst = AsSparseArray(DstSet);
+		const int32 OldNum = Dst.NumUnchecked();
 		int32 TmpIdx;
-		void* Tmp = AddConstructed(DstItems, /* out*/ TmpIdx);
+		void* Tmp = AddItem(Dst, /* out */ TmpIdx);
 		bool bSkip = true;
-		if constexpr (ItemKind == EPropertyKind::Range)
+		if constexpr (ElemKind == EPropertyKind::Range)
 		{
-			TConstArrayView<FRangeBinding> Inners = Base::GetInnerBindings();
-			for (FRangeLoadView Range : Src.AsRanges())
+			TConstArrayView<FRangeBinding> InnerBindings = Range.InnerBindings;
+			for (FRangeLoadView Item : Src.AsRanges())
 			{
-				LoadRange(Tmp, Range, Inners);
-				bSkip = HasElem(Dst, Tmp);
-				Tmp = bSkip ? Tmp : AddConstructed(DstItems, /* out*/ TmpIdx);
+				LoadRange(Tmp, Item, InnerBindings);
+				bSkip = HasElem(DstSet, Tmp);
+				Tmp = bSkip ? Tmp : AddItem(Dst, /* out */ TmpIdx);
 			}
 		}
 		else
 		{
-			for (FStructLoadView Struct : Src.AsStructs())
+			for (FStructLoadView Item : Src.AsStructs())
 			{
-				LoadStruct(Tmp, Struct);
-				bSkip = HasElem(Dst, Tmp);
-				Tmp = bSkip ? Tmp : AddConstructed(DstItems, /* out*/ TmpIdx);
+				LoadStruct(Tmp, Item);
+				bSkip = HasElem(DstSet, Tmp);
+				Tmp = bSkip ? Tmp : AddItem(Dst, /* out */ TmpIdx);
 			}
 		}
 
 		if (bSkip)
 		{
-			Base::DestructElem(Tmp);
-			DstItems.RemoveAtUninitialized(Layout.SparseArrayLayout, TmpIdx, 1);
+			Inner.DestroyItem(Tmp);
+			Dst.RemoveAtUninitialized(Layout.SparseArrayLayout, TmpIdx, 1);
 		}
 
-		if (DstItems.NumUnchecked() != OldNum)
+		if (Dst.NumUnchecked() != OldNum)
 		{
-			Base::Rehash(Dst);
+			Rehash(DstSet);
 		}
 	}
 
-	inline void Remove(FScriptSet& Dst, FRangeLoadView Src) const requires(bLeafItems)
+	inline void Remove(FScriptSet& Dst, FRangeLoadView Src) const requires(bLeaves)
 	{
 		for (LeafType Item : CastAs<LeafType>(Src.AsLeaves()))
 		{
@@ -1209,12 +1327,13 @@ struct TSetPropertyBinding : IItemRangeBinding, ICustomBinding, TSetPropertyBase
 	inline void Remove(FScriptSet& Dst, FRangeLoadView Src) const
 	{
 		TArray<uint8, TInlineAllocator<64>> Buffer;
-		Buffer.SetNumUninitialized(ElementSize);
-		void* Tmp = Base::ConstructElem(Buffer.GetData());
+		Buffer.SetNumUninitialized(Inner.Size);
+		Inner.InitItem(Buffer.GetData());
+		void* Tmp = Buffer.GetData();
 
-		if constexpr (ItemKind == EPropertyKind::Range)
+		if constexpr (ElemKind == EPropertyKind::Range)
 		{
-			TConstArrayView<FRangeBinding> Inners = Base::GetInnerBindings();
+			TConstArrayView<FRangeBinding> Inners = Range.InnerBindings;
 			for (FRangeLoadView Item : Src.AsRanges())
 			{
 				LoadRange(Tmp, Item, Inners);
@@ -1230,14 +1349,7 @@ struct TSetPropertyBinding : IItemRangeBinding, ICustomBinding, TSetPropertyBase
 			}
 		}
 
-		Base::DestructElem(Tmp);
-	}
-	
-	FTypedRange SaveDeltaRange(FSubSetIterator It, const FSaveContext& Ctx) const
-	{
-		FAnyRangeSaver Out(Ctx.Scratch, It.CountNum(), Base::BuiltItemSize);
-		Base::SaveDeltaItems(Out, It, Ctx);
-		return Out.Finalize(Base::MakeRangeSchema());
+		Inner.DestroyItem(Tmp);
 	}
 };
 
@@ -1260,16 +1372,16 @@ class FSetBindings
 	{
 		switch (GetPropertyKind(Elem))
 		{
-		case EPropertyKind::Range:			return BindNew<EPropertyKind::Range>(Property, Elem, Id);
-		case EPropertyKind::Struct:			return BindNew<EPropertyKind::Struct>(Property, Elem, Id);
-		case EPropertyKind::Bool:			return BindNew<EPropertyKind::Bool>(Property, Elem, Id);
-		case EPropertyKind::Integral8:		return BindNew<EPropertyKind::Integral8>(Property, Elem, Id);
-		case EPropertyKind::Integral16:		return BindNew<EPropertyKind::Integral16>(Property, Elem, Id);
-		case EPropertyKind::Integral32:		return BindNew<EPropertyKind::Integral32>(Property, Elem, Id);
-		case EPropertyKind::Integral64:		return BindNew<EPropertyKind::Integral64>(Property, Elem, Id);
-		case EPropertyKind::Float32:		return BindNew<EPropertyKind::Float32>(Property, Elem, Id);
-		case EPropertyKind::Float64:		return BindNew<EPropertyKind::Float64>(Property, Elem, Id);
-		default: check(false);				return nullptr;
+		case EPropertyKind::Range:		return BindNew<EPropertyKind::Range>(Property, Elem, Id);
+		case EPropertyKind::Struct:		return BindNew<EPropertyKind::Struct>(Property, Elem, Id);
+		case EPropertyKind::Bool:		return BindNew<EPropertyKind::Bool>(Property, Elem, Id);
+		case EPropertyKind::U8:			return BindNew<EPropertyKind::U8>(Property, Elem, Id);
+		case EPropertyKind::U16:		return BindNew<EPropertyKind::U16>(Property, Elem, Id);
+		case EPropertyKind::U32:		return BindNew<EPropertyKind::U32>(Property, Elem, Id);
+		case EPropertyKind::U64:		return BindNew<EPropertyKind::U64>(Property, Elem, Id);
+		case EPropertyKind::F32:		return BindNew<EPropertyKind::F32>(Property, Elem, Id);
+		case EPropertyKind::F64:		return BindNew<EPropertyKind::F64>(Property, Elem, Id);
+		default: check(false);			return nullptr;
 		}
 	}
 public:
@@ -1298,27 +1410,62 @@ static FSetBindings GSets;
 
 ////////////////////////////////////////////////////////////////////////////////////////////////
 
-enum class EMapFlags : uint32 { ConstructKey = 1, ConstructValue = 2, DestructKey = 4, DestructValue = 8 };
-
-// Currently can't extract constructor/destructor function pointers from FProperty, which
-// requires keeping FProperty* and prevents range binding reuse, @see AllocateMapBinding()
-template<class ScriptMap, EMapFlags Flags>
+// Flat TMap binding
+template<class ScriptMap, EPropertyKind KeyKind, EPropertyKind ValueKind>
 struct TMapPropertyBinding : IItemRangeBinding
 {
-	static constexpr bool bConstructKey =		(uint32(Flags) & uint32(EMapFlags::ConstructKey)) != 0;
-	static constexpr bool bConstructValue =		(uint32(Flags) & uint32(EMapFlags::ConstructValue)) != 0;
-	static constexpr bool bDestructKey =		(uint32(Flags) & uint32(EMapFlags::DestructKey)) != 0;
-	static constexpr bool bDestructValue =		(uint32(Flags) & uint32(EMapFlags::DestructValue)) != 0;
+	const FScriptMapLayout							Layout;
+	const TInnerProperty<KeyKind>					InnerKey;
+	const TInnerProperty<ValueKind>					InnerValue;
+
+	TMapPropertyBinding(FMapProperty* In)
+	: Layout(In->MapLayout)
+	, InnerKey(In->KeyProp)
+	, InnerValue(In->ValueProp)
+	{}
 
-	const FScriptMapLayout Layout;
-	FProperty* KeyProp;
-	FProperty* ValueProp;
+	inline SIZE_T GetStride() const	{ return Layout.SetLayout.Size; }
 
-	TMapPropertyBinding(FMapProperty* Property)
-	: Layout(Property->MapLayout)
-	, KeyProp(Property->KeyProp)
-	, ValueProp(Property->ValueProp)
-	{}
+	inline uint8* InitMap(ScriptMap& Map, int32 Num) const
+	{
+		uint8* It = SetNumUninitialized(Map, Layout, Num);
+		InitStridedItems(InnerKey, It, Num, GetStride());
+		InitStridedItems(InnerValue, It + Layout.ValueOffset, Num, GetStride());
+		return It;
+	}
+	
+	inline void Rehash(ScriptMap& Map) const
+	{
+		Map.Rehash(Layout, MakeHashFn(InnerKey));
+	}
+
+	inline void DestroyAll(ScriptMap& Map) const
+	{
+		if (InnerKey.bDestruct || InnerValue.bDestruct)
+		{
+			const SIZE_T Stride = GetStride();
+			const int32 ValueOffset = Layout.ValueOffset;
+			const int32 Num = Map.NumUnchecked();
+			uint8* It = static_cast<uint8*>(Map.GetData(0, Layout));
+			if (IsCompact(Map))
+			{
+				DestroyStridedItems(InnerKey, It, Num, Stride);
+				DestroyStridedItems(InnerValue, It + ValueOffset, Num, Stride);
+			}
+			else
+			{
+				for (int32 Idx = 0, Max = Map.GetMaxIndex(); Idx < Max; ++Idx)
+				{
+					if (Map.IsValidIndex(Idx))
+					{
+						InnerKey.DestroyItem(It);
+						InnerValue.DestroyItem(It + ValueOffset);
+					}
+					It += Stride;
+				}
+			}
+		}
+	}
 
 	virtual void ReadItems(FSaveRangeContext& Ctx) const override
 	{
@@ -1328,135 +1475,321 @@ struct TMapPropertyBinding : IItemRangeBinding
 	virtual void MakeItems(FLoadRangeContext& Ctx) const override
 	{
 		ScriptMap& Map = Ctx.Request.GetRange<ScriptMap>();
+		int32 NewNum = static_cast<int32>(Ctx.Request.NumTotal());
 
 		if (Ctx.Request.IsFirstCall())
 		{
-			DestructItems(Map);
-
-			int32 NewNum = static_cast<int32>(Ctx.Request.NumTotal());
+			DestroyAll(Map);
 			Map.Empty(NewNum, Layout);
 			if (NewNum)
 			{
-				// Todo: Bug, mimic TSetPropertyBinding::MakeItems
-				void* Data = AddConstructedItems(Map, NewNum);
-				Ctx.Items.Set(Data, Ctx.Request.NumTotal(), ItemSize());
+				void* Items = InitMap(Map, NewNum);
+				Ctx.Items.Set(Items, Ctx.Request.NumTotal(), GetStride());
 				Ctx.Items.RequestFinalCall();
 			}
 		}
 		else
 		{
 			check(Ctx.Request.IsFinalCall());
-			Map.Rehash(Layout, [this](const void* Src) { return KeyProp->GetValueTypeHash(Src); });		
+			Rehash(Map);
 		}
 	}
+};
+
+////////////////////////////////////////////////////////////////////////////////////////////////
+
+struct FMapMemberBindings
+{
+	FMemberBinding Key;
+	FMemberBinding Value;
+	FMemberBinding Pair;
+};
 
-	void* AddConstructedItems(ScriptMap& Map, const int32 Num) const
+template <EPropertyKind KeyKind, EPropertyKind ValueKind>
+struct TMapPropertyDeltaBinding : TMapPropertyBinding<FScriptMap, KeyKind, ValueKind>, ICustomBinding
+{
+	using Super = TMapPropertyBinding<FScriptMap, KeyKind, ValueKind>;
+	using Super::Layout;
+	using Super::InnerKey;
+	using Super::InnerValue;
+	using Super::GetStride;
+	using Super::InitMap;
+	using Super::Rehash;
+	
+	const TPropertyRangeSerializer<KeyKind>			KeyRange;
+	const TPropertyRangeSerializer<ValueKind>		ValueRange;
+	const FStructRangeSerializer					PairRange;
+
+	TMapPropertyDeltaBinding(FMapProperty* Map, FMapMemberBindings Members)
+	: Super(Map)
+	, KeyRange(Members.Key)
+	, ValueRange(Members.Value)
+	, PairRange(Members.Pair)
+	{}
+
+	inline const void* GetValue(const void* Pair) const
 	{
-		for (int32 Idx = 0; Idx < Num; ++Idx)
+		return static_cast<const uint8*>(Pair) + Layout.ValueOffset;
+	} 
+
+	inline int32 FindKey(const FScriptMap& Map, const void* Key) const
+	{
+		return Map.FindPairIndex(Key, Layout, MakeHashFn(InnerKey), MakeIdenticalFn(InnerKey));
+	}
+
+	inline bool HasKey(const FScriptMap& Map, const void* Key) const
+	{
+		return FindKey(Map, Key) != INDEX_NONE;
+	}
+
+	inline bool HasPair(const FScriptMap& Map, const void* Pair) const
+	{
+		const void* Key = Pair;
+		if (int32 Idx = FindKey(Map, Key); Idx != INDEX_NONE)
+		{
+			const void* FoundPair = Map.GetData(Idx, Layout);
+			return MakeIdenticalFn(InnerValue)(GetValue(Pair), GetValue(FoundPair));
+		}
+		return false;
+	}
+
+	inline bool HasElem(const FScriptMap& Map, const void* Pair) const
+	{
+		return HasPair(Map, Pair);
+	}
+
+	inline uint8* AddPair(FScriptSparseArray& Dst, int32& OutIdx, int32 ValueOffset) const
+	{
+		OutIdx = Dst.AddUninitialized(Layout.SetLayout.SparseArrayLayout);
+		uint8* Out = static_cast<uint8*>(Dst.GetData(OutIdx, Layout.SetLayout.SparseArrayLayout));
+		InnerKey.InitItem(Out);
+		InnerValue.InitItem(Out + ValueOffset);
+		return Out;
+	}
+
+	inline void RemoveKey(FScriptMap& Map, const void* Key) const
+	{
+		if (int32 Idx = FindKey(Map, Key); Idx != INDEX_NONE)
 		{
-			Map.AddUninitialized(Layout);
+			Map.RemoveAt(Idx, Layout);
 		}
-		check(IsCompact(Map));
+	}
 
-		uint8* Data = static_cast<uint8*>(Map.GetData(0, Layout));
-		SIZE_T Stride = ItemSize();
-		if constexpr (bConstructKey && bConstructValue == false)
+	virtual void SaveCustom(FMemberBuilder& Dst, const void* InSrc, const void* InDefault, const FSaveContext& Ctx) override
+	{
+		const FScriptMap& Src = *static_cast<const FScriptMap*>(InSrc);
+		const FScriptMap* Default = static_cast<const FScriptMap*>(InDefault);
+		
+		if (Src.IsEmpty())
 		{
-			FMemory::Memzero(Data, Stride * Num);
+			if (Default && !Default->IsEmpty())
+			{
+				Dst.AddRange(GUE.CommonIds.Del, SaveAll(*Default, Layout, KeyRange, Ctx));
+			}
 		}
-		if constexpr (bConstructKey)
+		else if (Default && !Default->IsEmpty())
 		{
-			InitializeValues(KeyProp, Data, Num, Stride);
+			TBitArray<> DeletedIds(false, Default->GetMaxIndex());
+			for (int32 Idx = 0, Max = Default->GetMaxIndex(); Idx < Max; ++Idx)
+			{
+				DeletedIds[Idx] = Default->IsValidIndex(Idx) && HasKey(Src, Default->GetData(Idx, Layout));
+			}
+			if (FSubMapIterator Deleted{Layout, *Default, DeletedIds})
+			{
+				Dst.AddRange(GUE.CommonIds.Del, SaveSome(Deleted, KeyRange, Ctx));
+			}
+			
+			TBitArray<> AddedIds(false, Src.GetMaxIndex());
+			for (int32 Idx = 0, Max = Src.GetMaxIndex(); Idx < Max; ++Idx)
+			{
+				AddedIds[Idx] = Src.IsValidIndex(Idx) && !HasPair(*Default, Src.GetData(Idx, Layout));
+			}
+			if (FSubMapIterator Added{Layout, Src, AddedIds})
+			{
+				Dst.AddRange(GUE.CommonIds.Add, SaveSome(Added, PairRange, Ctx));
+			}
 		}
-		if constexpr (bConstructValue)
+		else
 		{
-			InitializeValues(ValueProp, Data + Layout.ValueOffset, Num, Stride);
+			Dst.AddRange(GUE.CommonIds.Add, SaveAll(Src, Layout, PairRange, Ctx));
 		}
+	}
 
-		return Data;
+	virtual void LoadCustom(void* Dst, FStructLoadView Src, ECustomLoadMethod Method) const override
+	{
+		LoadSetDelta(*this, *static_cast<FScriptMap*>(Dst), Src, Method);
 	}
 
-	inline void DestructItems(ScriptMap& Map) const
+	// Load into empty set
+	inline void Append(FScriptMap& Dst, FRangeLoadView Src) const
 	{
-		if constexpr (bDestructKey || bDestructValue)
+		const int32 ValueOffset = Layout.ValueOffset;
+		const SIZE_T Stride = GetStride();
+
+		uint8* It = InitMap(Dst, Src.Num());
+		for (FStructLoadView Struct : Src.AsStructs())
 		{
-			if (const int32 Num = Map.Num())
-			{
-				if (IsCompact(Map))
-				{
-					uint8* Data = static_cast<uint8*>(Map.GetData(0, Layout));
-					if constexpr (bDestructKey)
-					{
-						DestroyValues(KeyProp, Data, Num, ItemSize());
-					}
-					if constexpr (bDestructValue)
-					{
-						DestroyValues(ValueProp, Data + Layout.ValueOffset, Num, ItemSize());
-					}
-				}
-				else
-				{
-					for (int32 Idx = 0; Idx < Num; ++Idx)
-					{
-						uint8* Item = static_cast<uint8*>(Map.GetData(Idx, Layout));
-						if constexpr (bDestructKey)
-						{
-							DestroyValue(KeyProp, Item);
-						}
-						if constexpr (bDestructValue)
-						{
-							DestroyValue(ValueProp, Item + Layout.ValueOffset);
-						}
-					}
-				}
-			}
+			// Equivalent to LoadStruct(It, Struct);
+			FBitCacheReader Bits;
+			KeyRange.LoadItem(It, /* in-out */ Struct.Values, /* in-out */ Bits, Struct.Schema.Batch);
+			ValueRange.LoadItem(It + ValueOffset, /* in-out */ Struct.Values, /* in-out */ Bits, Struct.Schema.Batch);
+
+			It += Stride;
+		}
+
+		Rehash(Dst);
+	}
+
+	// Load structs or ranges into non-empty map
+	inline void Insert(FScriptMap& Dst, FRangeLoadView Src) const
+	{
+		// Written to avoid FProperty::CopyCompleteValue_InContainer dependency
+		// Items are loaded directly into sparse array and then removed if a duplicate existed
+
+		FScriptSparseArray& DstArray = AsSparseArray(Dst);
+		const int32 OldNum = DstArray.NumUnchecked();
+		const int32 ValueOffset = Layout.ValueOffset;
+		int32 TmpIdx;
+		uint8* Tmp = AddPair(DstArray, /* out*/ TmpIdx, ValueOffset);
+		bool bSkip = true;
+
+		for (FStructLoadView Struct : Src.AsStructs())
+		{
+			// Equivalent to LoadStruct(It, Struct);
+			FBitCacheReader Bits;
+			KeyRange.LoadItem(Tmp,					/* in-out */ Struct.Values, /* in-out */ Bits, Struct.Schema.Batch);
+			ValueRange.LoadItem(Tmp + ValueOffset,	/* in-out */ Struct.Values, /* in-out */ Bits, Struct.Schema.Batch);
+			bSkip = HasPair(Dst, Tmp);
+			Tmp = bSkip ? Tmp : AddPair(DstArray, /* out*/ TmpIdx, ValueOffset);
+		}
+
+		if (bSkip)
+		{
+			InnerKey.DestroyItem(Tmp);
+			InnerValue.DestroyItem(Tmp + ValueOffset);
+			DstArray.RemoveAtUninitialized(Layout.SetLayout.SparseArrayLayout, TmpIdx, 1);
+		}
+
+		if (DstArray.NumUnchecked() != OldNum)
+		{
+			Rehash(Dst);
+		}
+	}
+
+	inline void Remove(FScriptMap& Dst, FRangeLoadView Src) const
+	{
+		using LeafType = EquivalentLeafType<KeyKind>;
+		for (LeafType Item : CastAs<LeafType>(Src.AsLeaves()))
+		{
+			RemoveKey(Dst, &Item);
+		}
+	}
+
+	inline void Remove(FScriptMap& Dst, FRangeLoadView Src) const requires (KeyKind == EPropertyKind::Range)
+	{
+		TArray<uint8, TInlineAllocator<64>> Buffer;
+		Buffer.SetNumUninitialized(InnerKey.Size);
+		void* Tmp = Buffer.GetData();
+
+		InnerKey.InitItem(Tmp);
+		for (FRangeLoadView Item : Src.AsRanges())
+		{
+			LoadRange(Tmp, Item, KeyRange.InnerBindings);
+			RemoveKey(Dst, Tmp);
+		}
+		InnerKey.DestroyItem(Tmp);
+	}
+
+	inline void Remove(FScriptMap& Dst, FRangeLoadView Src) const requires (KeyKind == EPropertyKind::Struct)
+	{
+		TArray<uint8, TInlineAllocator<64>> Buffer;
+		Buffer.SetNumUninitialized(InnerKey.Size);
+		void* Tmp = Buffer.GetData();
+
+		InnerKey.InitItem(Tmp);
+		for (FStructLoadView Item : Src.AsStructs())
+		{
+			LoadStruct(Tmp, Item);
+			RemoveKey(Dst, Tmp);
 		}
+		InnerKey.DestroyItem(Tmp);
 	}
 
-	inline SIZE_T ItemSize() const { return Layout.SetLayout.SparseArrayLayout.Size; }
+	virtual bool DiffCustom(const void* A, const void* B) const override
+	{
+		return DiffSet(*this, *static_cast<const FScriptMap*>(A), *static_cast<const FScriptMap*>(B));
+	}
 };
 
-inline uint32 GetMapPropertyFlags(FProperty* Key, FProperty* Val)
+////////////////////////////////////////////////////////////////////////////////////////////////
+
+class FMapBindings
 {
-	return	HasConstructor(Key) * static_cast<uint32>(EMapFlags::ConstructKey) |
-			HasConstructor(Val) * static_cast<uint32>(EMapFlags::ConstructValue) |
-			HasDestructor(Key) *  static_cast<uint32>(EMapFlags::DestructKey) |
-			HasDestructor(Val) *  static_cast<uint32>(EMapFlags::DestructValue);
-}
+	TMap<FStructSchemaId, FRangeBinding> NormalBindings;
+	TMap<FStructSchemaId, FRangeBinding> FrozenBindings;
 
-template<class ScriptMap>
-inline IItemRangeBinding* AllocateBinding(FMapProperty* Property)
-{
-	switch (GetMapPropertyFlags(Property->KeyProp, Property->ValueProp))
-	{
-	case 0:		return new TMapPropertyBinding<ScriptMap, EMapFlags(0)>(Property);
-	case 1:		return new TMapPropertyBinding<ScriptMap, EMapFlags(1)>(Property);
-	case 2:		return new TMapPropertyBinding<ScriptMap, EMapFlags(2)>(Property);
-	case 3:		return new TMapPropertyBinding<ScriptMap, EMapFlags(3)>(Property);
-	case 4:		return new TMapPropertyBinding<ScriptMap, EMapFlags(4)>(Property);
-	case 5:		return new TMapPropertyBinding<ScriptMap, EMapFlags(5)>(Property);
-	case 6:		return new TMapPropertyBinding<ScriptMap, EMapFlags(6)>(Property);
-	case 7:		return new TMapPropertyBinding<ScriptMap, EMapFlags(7)>(Property);
-	case 8:		return new TMapPropertyBinding<ScriptMap, EMapFlags(8)>(Property);
-	case 9:		return new TMapPropertyBinding<ScriptMap, EMapFlags(9)>(Property);
-	case 10:	return new TMapPropertyBinding<ScriptMap, EMapFlags(10)>(Property);
-	case 11:	return new TMapPropertyBinding<ScriptMap, EMapFlags(11)>(Property);
-	case 12:	return new TMapPropertyBinding<ScriptMap, EMapFlags(12)>(Property);
-	case 13:	return new TMapPropertyBinding<ScriptMap, EMapFlags(13)>(Property);
-	case 14:	return new TMapPropertyBinding<ScriptMap, EMapFlags(14)>(Property);
-	case 15:	return new TMapPropertyBinding<ScriptMap, EMapFlags(15)>(Property);
-	default:	check(false);
-	}
-	return nullptr;
-}
+	template<EPropertyKind KeyKind, EPropertyKind ValueKind>
+	static IItemRangeBinding* New3(FMapProperty* Property, FMapMemberBindings Members, bool bFreezable)
+	{
+		if (bFreezable) // Freezable maps aren't delta serialized
+		{
+			return new TMapPropertyBinding<FFreezableScriptMap, KeyKind, ValueKind>(Property);
+		}
+		return new TMapPropertyDeltaBinding<KeyKind, ValueKind>(Property, Members);
+	}
 
-static FRangeBinding AllocateMapBinding(FMapProperty* Property)
-{
-	bool bFreezable = EnumHasAnyFlags(Property->MapFlags, EMapPropertyFlags::UsesMemoryImageAllocator);
-	// Todo: Ownership / memory leak
-	IItemRangeBinding* Leak = bFreezable ? AllocateBinding<FFreezableScriptMap>(Property) : AllocateBinding<FScriptMap>(Property);
-	return FRangeBinding(*Leak, DefaultSizeType);
-}
+	template<EPropertyKind KeyKind>
+	inline IItemRangeBinding* New2(FMapProperty* Property, FMapMemberBindings Members, bool bFreezable)
+	{
+		switch (GetPropertyKind(Members.Value))
+		{
+		case EPropertyKind::Range:		return New3<KeyKind, EPropertyKind::Range>(Property, Members, bFreezable);
+		case EPropertyKind::Struct:		return New3<KeyKind, EPropertyKind::Struct>(Property, Members, bFreezable);
+		case EPropertyKind::Bool:		return New3<KeyKind, EPropertyKind::Bool>(Property, Members, bFreezable);
+		case EPropertyKind::U8:			return New3<KeyKind, EPropertyKind::U8>(Property, Members, bFreezable);
+		case EPropertyKind::U16:		return New3<KeyKind, EPropertyKind::U16>(Property, Members, bFreezable);
+		case EPropertyKind::U32:		return New3<KeyKind, EPropertyKind::U32>(Property, Members, bFreezable);
+		case EPropertyKind::U64:		return New3<KeyKind, EPropertyKind::U64>(Property, Members, bFreezable);
+		case EPropertyKind::F32:		return New3<KeyKind, EPropertyKind::F32>(Property, Members, bFreezable);
+		case EPropertyKind::F64:		return New3<KeyKind, EPropertyKind::F64>(Property, Members, bFreezable);
+		default:						check(false); return nullptr;
+		}
+	}
+
+	inline IItemRangeBinding* New(FMapProperty* Property, FMapMemberBindings Members, bool bFreezable)
+	{
+		switch (GetPropertyKind(Members.Key))
+		{
+		case EPropertyKind::Range:		return New2<EPropertyKind::Range>(Property, Members, bFreezable);
+		case EPropertyKind::Struct:		return New2<EPropertyKind::Struct>(Property, Members, bFreezable);
+		case EPropertyKind::Bool:		return New2<EPropertyKind::Bool>(Property, Members, bFreezable);
+		case EPropertyKind::U8:			return New2<EPropertyKind::U8>(Property, Members, bFreezable);
+		case EPropertyKind::U16:		return New2<EPropertyKind::U16>(Property, Members, bFreezable);
+		case EPropertyKind::U32:		return New2<EPropertyKind::U32>(Property, Members, bFreezable);
+		case EPropertyKind::U64:		return New2<EPropertyKind::U64>(Property, Members, bFreezable);
+		case EPropertyKind::F32:		return New2<EPropertyKind::F32>(Property, Members, bFreezable);
+		case EPropertyKind::F64:		return New2<EPropertyKind::F64>(Property, Members, bFreezable);
+		default:						check(false); return nullptr;
+		}
+	}
+
+public:
+	FRangeBinding Bind(FMapProperty* Property, FMapMemberBindings Members)
+	{
+		FStructSchemaId PairId = static_cast<FStructSchemaId>(Members.Pair.InnermostSchema.Get());
+		bool bFreezable = EnumHasAnyFlags(Property->MapFlags, EMapPropertyFlags::UsesMemoryImageAllocator);
+		TMap<FStructSchemaId, FRangeBinding>& Bindings = bFreezable ? FrozenBindings : NormalBindings;
+		if (const FRangeBinding* RangeBinding = Bindings.Find(PairId))
+		{
+			return *RangeBinding;
+		}
+
+		// Todo: Ownership / memory leak
+		IItemRangeBinding* Leak = New(Property, Members, bFreezable);
+		return Bindings.Emplace(PairId, FRangeBinding(*Leak, DefaultSizeType));
+	}
+};
+static FMapBindings GMaps;
 
 ////////////////////////////////////////////////////////////////////////////////////////////////
 
@@ -1503,7 +1836,6 @@ public:
 		return Out;
 	}
 };
-
 static FPairBindings GPairs;
 
 //////////////////////////////////////////////////////////////////////////////////////////////
@@ -1745,13 +2077,15 @@ private:
 	}
 	inline FMemberBinding BindMap(FMapProperty* Property)
 	{
-		FMemberBinding PairBinding = GPairs.Bind(BindSingleProperty(Property->KeyProp), BindSingleProperty(Property->ValueProp));
-		return BindAsRange(Property, AllocateMapBinding(Property), PairBinding);
+		FMemberBinding Key = BindSingleProperty(Property->KeyProp);
+		FMemberBinding Value = BindSingleProperty(Property->ValueProp);
+		FMemberBinding Pair = GPairs.Bind(Key, Value);
+		return BindAsRange(Property, GMaps.Bind(Property, {Key, Value, Pair}), Pair);
 	}
 	inline FMemberBinding BindSet(FSetProperty* Property)
 	{
-		FMemberBinding ElemBinding = BindSingleProperty(Property->ElementProp);
-		return BindAsRange(Property, GSets.Bind(Property, ElemBinding), ElemBinding);
+		FMemberBinding Elem = BindSingleProperty(Property->ElementProp);
+		return BindAsRange(Property, GSets.Bind(Property, Elem), Elem);
 	}
 	inline FMemberBinding BindOptional(FOptionalProperty* Property)
 	{
