PlainProps new default struct lookup API and implementation for UScriptStruct defaults

#rnx

diff --git a/Engine/Plugins/Experimental/PlainProps/Source/Private/PlainPropsSave.cpp b/Engine/Plugins/Experimental/PlainProps/Source/Private/PlainPropsSave.cpp
--- a/Engine/Plugins/Experimental/PlainProps/Source/Private/PlainPropsSave.cpp
+++ b/Engine/Plugins/Experimental/PlainProps/Source/Private/PlainPropsSave.cpp
@@ -123,6 +123,8 @@ struct TNonLeafRangeSaver
 	}
 };
 
+struct FDefaultStruct { FStructSchemaId Id; const void* Struct; };
+
 using FInternalNestedRangeSaver = TNonLeafRangeSaver<FBuiltRange*, FRangeMemberBinding>;
 using FInternalStructRangeSaver = TNonLeafRangeSaver<FBuiltStruct*, FStructSchemaId>;
 using FInternalStructRangeDeltaSaver = TNonLeafRangeSaver<FBuiltStruct*, FDefaultStruct>;
@@ -198,9 +200,10 @@ FBuiltRange* SaveLeafRange(const void* Range, const ILeafRangeBinding& Binding,
 
 [[nodiscard]] static FBuiltRange* SaveStructRange(const void* Range, const IItemRangeBinding& ItemBinding, const FSaveContext& Ctx, FStructSchemaId Id)
 {
-	if (const FDefaultStruct* Defaults = Algo::FindBy(Ctx.Defaults, Id, &FDefaultStruct::Id))
+	if (Ctx.Defaults)
 	{
-		return SaveNonLeafRange<FInternalStructRangeDeltaSaver>(Range, ItemBinding, Ctx, *Defaults);
+		FDefaultStruct Default = { Id, Ctx.Defaults->Get(Id) };
+		return SaveNonLeafRange<FInternalStructRangeDeltaSaver>(Range, ItemBinding, Ctx, Default);
 	}
 	else
 	{
diff --git a/Engine/Plugins/Experimental/PlainProps/Source/Public/PlainPropsBind.h b/Engine/Plugins/Experimental/PlainProps/Source/Public/PlainPropsBind.h
--- a/Engine/Plugins/Experimental/PlainProps/Source/Public/PlainPropsBind.h
+++ b/Engine/Plugins/Experimental/PlainProps/Source/Public/PlainPropsBind.h
@@ -564,6 +564,14 @@ private:
 
 ////////////////////////////////////////////////////////////////////////////////////////////////
 
+// Lookup default instances when saving struct ranges
+struct IDefaultStructs
+{
+	virtual const void* Get(FStructSchemaId Id) = 0;
+};
+
+////////////////////////////////////////////////////////////////////////////////////////////////
+
 struct FStructBindIds : IStructBindIds
 {
 	FStructBindIds(const FCustomBindings& InCustoms, const FSchemaBindings& InSchemas)
diff --git a/Engine/Plugins/Experimental/PlainProps/Source/Public/PlainPropsSave.h b/Engine/Plugins/Experimental/PlainProps/Source/Public/PlainPropsSave.h
--- a/Engine/Plugins/Experimental/PlainProps/Source/Public/PlainPropsSave.h
+++ b/Engine/Plugins/Experimental/PlainProps/Source/Public/PlainPropsSave.h
@@ -11,27 +11,23 @@ struct FBuiltRange;
 struct FBuiltStruct;
 class FCustomBindings;
 class FDeclarations;
+struct IDefaultStructs;
 struct ILeafRangeBinding;
 struct FRangeMemberBinding;
 class FSchemaBindings;
 class FScratchAllocator;
 
-// Temporary data structure, will be replaced by something more sophisticated
-// perhaps deduplicating all zero-memory defaults
-struct FDefaultStruct { FStructSchemaId Id; const void* Struct; };
-using FDefaultStructs = TConstArrayView<FDefaultStruct>;
-
 struct FSaveContext
 {
 	const FDeclarations&		Declarations;
 	const FSchemaBindings&		Schemas;
 	FCustomBindings&			Customs;
 	FScratchAllocator&			Scratch;
-	FDefaultStructs				Defaults;
+	IDefaultStructs*			Defaults = nullptr;
 };
 
 template<typename Runtime>
-FSaveContext MakeSaveContext(FDefaultStructs Defaults, FScratchAllocator& Scratch)
+FSaveContext MakeSaveContext(FScratchAllocator& Scratch, IDefaultStructs* Defaults)
 {
 	return { Runtime::GetTypes(), Runtime::GetSchemas(), Runtime::GetCustoms(), Scratch, Defaults };
 }
diff --git a/Engine/Plugins/Experimental/PlainPropsUObject/Source/Private/PlainPropsUObjectBindings.cpp b/Engine/Plugins/Experimental/PlainPropsUObject/Source/Private/PlainPropsUObjectBindings.cpp
--- a/Engine/Plugins/Experimental/PlainPropsUObject/Source/Private/PlainPropsUObjectBindings.cpp
+++ b/Engine/Plugins/Experimental/PlainPropsUObject/Source/Private/PlainPropsUObjectBindings.cpp
@@ -10,6 +10,7 @@
 #include "Misc/CoreDelegates.h"
 #include "Misc/DefinePrivateMemberPtr.h"
 #include "UObject/AnsiStrProperty.h"
+#include "UObject/Class.h"
 #include "UObject/EnumProperty.h"
 #include "UObject/FieldPathProperty.h"
 #include "UObject/PropertyOptional.h"
@@ -24,6 +25,100 @@
 namespace PlainProps::UE
 {
 
+struct FDefaultStruct
+{
+	UScriptStruct::ICppStructOps&		Ops;
+	alignas(16) uint8					Instance[0];
+};
+
+static FDefaultStruct* NewDefault(UScriptStruct::ICppStructOps& Ops)
+{
+	check(Ops.GetAlignment() <= 16);
+	uint32 Size = sizeof(FDefaultStruct) + Ops.GetSize();
+	FDefaultStruct Header = {Ops};
+	FDefaultStruct* Out = new (FMemory::MallocZeroed(Size)) FDefaultStruct(Header);
+	Ops.Construct(Out->Instance);
+	return Out;
+}
+
+static void DeleteDefault(uint8* Instance)
+{
+	FDefaultStruct* Struct = reinterpret_cast<FDefaultStruct*>(Instance - offsetof(FDefaultStruct, Instance));
+	if (Struct->Ops.HasDestructor())
+	{
+		Struct->Ops.Destruct(Instance);
+	}
+	FMemory::Free(Struct);
+}
+
+static void ReserveZeroes(/* in-out */ FMutableMemoryView& Zeroes, SIZE_T Size, uint32 Alignment)
+{
+	Size +=  FMath::Max<int32>(0, Alignment - 16);
+	Size = Align(Size, 4096);
+	if (Zeroes.GetSize() < Size)
+	{
+		FMemory::Free(Zeroes.GetData());
+		Zeroes = FMutableMemoryView(FMemory::MallocZeroed(Size, 16), Size);
+	}
+}
+
+FDefaultStructs::~FDefaultStructs()
+{
+	for (TPair<FStructSchemaId, uint8*> Instance : Instances)
+	{
+		DeleteDefault(Instance.Value);
+	}
+}
+
+void FDefaultStructs::Bind(FStructSchemaId Id, const UScriptStruct* Struct)
+{
+	const EStructFlags Flags = Struct->StructFlags;
+	const SIZE_T Size = Struct->GetStructureSize();
+	const uint32 Alignment = Struct->GetMinAlignment();
+
+	if (Id.Idx >= static_cast<uint32>(Instanced.Num()))
+	{
+		Instanced.SetNum(FMath::RoundUpToPowerOfTwo64(Id.Idx + 1), false);
+	}
+
+	UScriptStruct::ICppStructOps* Ops = Struct->GetCppStructOps();
+	if (!!(Flags & STRUCT_ZeroConstructor) || Ops == nullptr)
+	{
+		ReserveZeroes(Zeroes, Size, Alignment);
+	}
+	else
+	{
+		check(Ops->GetSize() == Size);
+		FDefaultStruct* Default = NewDefault(*Ops);
+		if (FMemory::MemIsZero(Default->Instance, Size))
+		{
+			DeleteDefault(Default->Instance);
+			ReserveZeroes(Zeroes, Size, Alignment);
+		}
+		else
+		{
+			Instanced[Id.Idx] = true;
+			Instances.Add(Id, Default->Instance);
+		}
+	}
+}
+
+void FDefaultStructs::Drop(FStructSchemaId Id)
+{
+	if (Instanced[Id.Idx])
+	{
+		Instanced[Id.Idx] = false;
+		DeleteDefault(Instances.FindAndRemoveChecked(Id));
+	}
+}
+
+const void* FDefaultStructs::Get(FStructSchemaId Id)
+{
+	return Instanced[Id.Idx] ? Instances.FindChecked(Id) : Zeroes.GetData(); 
+}
+
+////////////////////////////////////////////////////////////////////////////////////////////////
+
 FCommonUObjectIds::FCommonUObjectIds(TIdIndexer<FName>& Names)
 : Core(				Names.MakeScope("/Script/Core"))
 , CoreUObject(		Names.MakeScope("/Script/CoreUObject"))
@@ -2229,7 +2324,13 @@ void BindStruct(FStructSchemaId Id, UStruct* Struct)
 	{
 		Binder.BindMember(Property);
 	}
-	GUE.Schemas.BindStruct(Id, Id, Binder.GetMembers()); 
+	GUE.Schemas.BindStruct(Id, Id, Binder.GetMembers());
+	
+	// Don't bind CDOs, object defaults are passed in from top and objects aren't owned by containers 
+	if (Struct->HasAnyCastFlags(CASTCLASS_UScriptStruct))
+	{
+		GUE.Defaults.Bind(Id, static_cast<const UScriptStruct*>(Struct));
+	}
 }
 
 ////////////////////////////////////////////////////////////////////////////////////////////////
diff --git a/Engine/Plugins/Experimental/PlainPropsUObject/Source/Public/PlainPropsUObjectRuntime.h b/Engine/Plugins/Experimental/PlainPropsUObject/Source/Public/PlainPropsUObjectRuntime.h
--- a/Engine/Plugins/Experimental/PlainPropsUObject/Source/Public/PlainPropsUObjectRuntime.h
+++ b/Engine/Plugins/Experimental/PlainPropsUObject/Source/Public/PlainPropsUObjectRuntime.h
@@ -6,11 +6,28 @@
 #include "PlainPropsDeclare.h"
 #include "PlainPropsIndex.h"
 #include "PlainPropsTypes.h"
+#include "Containers/Map.h"
 #include "UObject/NameTypes.h"
 
 namespace PlainProps::UE
 {
 
+struct FDefaultStruct;
+
+class FDefaultStructs final : IDefaultStructs
+{
+public:
+	~FDefaultStructs();
+	void								Bind(FStructSchemaId Id, const UScriptStruct* Struct);
+	void								Drop(FStructSchemaId Id);
+	virtual const void*					Get(FStructSchemaId Id) override;
+
+private:
+	FMutableMemoryView					Zeroes;
+	TBitArray<>							Instanced;
+	TMap<FStructSchemaId, uint8*>		Instances;
+};
+
 struct FCommonUObjectIds
 {
 	explicit FCommonUObjectIds(TIdIndexer<FName>& Names);
@@ -48,6 +65,7 @@ struct FGlobals
 	FCustomBindings				Customs;
 	FCustomBindings				DeltaCustoms;
 	FCommonUObjectIds			CommonIds;
+	FDefaultStructs				Defaults;
 };
 PLAINPROPSUOBJECT_API extern FGlobals GUE;
 
