PlainProps bindings and bug fixes to save all FProperties in all UObjects during ShooterGame startup

Add Core bindings for FGuid, FColor, FLinearColor
Add CoreUObject bindings for FScriptDelegate, FMulticastScriptDelegate, FSparseDelegate, FFieldPath, FScriptInterface
Add batch bindings for FText, FName, FObjectPtr, FSoftObjectPtr, FWeakObjectPtr, FLazyObjectPtr
Add default struct bindings for all of the above
Add FMemoryPropertyBatch
Add test batch saver and code that saves  all UObjects
Fix bunch of bugs in previously unexercised code
Better custom binding overlays

#rnx

diff --git a/Engine/Plugins/Experimental/PlainProps/Source/Private/PlainPropsBind.cpp b/Engine/Plugins/Experimental/PlainProps/Source/Private/PlainPropsBind.cpp
--- a/Engine/Plugins/Experimental/PlainProps/Source/Private/PlainPropsBind.cpp
+++ b/Engine/Plugins/Experimental/PlainProps/Source/Private/PlainPropsBind.cpp
@@ -157,69 +157,99 @@ void* FLeafRangeAllocator::Allocate(uint64 Num, SIZE_T LeafSize)
 
 ////////////////////////////////////////////////////////////////////////////////////////////////
 
-static bool operator==(FInnerStruct A, FInnerStruct B) { return A.Name == B.Name && A.BindId == B.BindId; }
-
-void FCustomBindings::BindStruct(FStructSchemaId BindId, FStructSchemaId DeclId, TConstArrayView<FInnerStruct> LoweredInners, ICustomBinding& Binding)
+inline bool Cull(FIdWindow Window, uint32 Id)
 {
-	checkf(!Find(BindId), TEXT("'%s' already bound"), *Debug.Print(BindId));
-	FInnersHandle Inners = StoreInners(LoweredInners);
-	check(Algo::Compare(LoweredInners, FetchInners(Inners)));
-	Entries.Emplace(BindId, DeclId, &Binding, Inners);
+	return Id - Window.Min < Window.Num;
 }
 
-const ICustomBinding* FCustomBindings::FindStruct(FStructSchemaId BindId, TConstArrayView<FInnerStruct>& OutLoweredInners) const
+static void UpdateCulling(FIdWindow& Window, uint32 Id)
 {
-	FCustomBindingEntry Entry = Find(BindId);
-	OutLoweredInners = FetchInners(Entry.LoweredInners);
-	return Entry.Binding;
+	if (Window.Num == 0)
+	{
+		Window = { Id, 1 };
+	}
+	else if (Id < Window.Min)
+	{
+		Window = { Id, Window.Num + Window.Min - Id };
+	}
+	else 
+	{
+		Window.Num = FMath::Max(1u + Id - Window.Min, Window.Num);
+	}
 }
 
-void FCustomBindings::DropStruct(FStructSchemaId BindId)
+struct FCustomBindingHandle
 {
-	for (FCustomBindingEntry& Entry : Entries)
+	ICustomBinding**				Bindings = nullptr;
+	uint32							Max = 0;
+
+	explicit						operator bool() { return !!Bindings;}
+	const FStructDeclaration**		GetDeclarations() { return reinterpret_cast<const FStructDeclaration**>((void*)(Bindings + Max)); }
+	FInnersHandle*					GetLoweredInners() { return reinterpret_cast<FInnersHandle*>(Bindings + 2 * Max); }
+};
+
+void FCustomBindingMap::Bind(FStructSchemaId BindId, ICustomBinding& Binding, const FStructDeclaration& Declaration, FInnersHandle LoweredInners)
+{
+	// Add key
+	bool bExists = false;
+	int32 SetIdx = Keys.Add(BindId, &bExists).AsInteger();
+	checkf(!bExists, TEXT("'%s' already bound"), *Debug.Print(BindId));
+
+	// Grow values if needed
+	uint32 MaxKeys = static_cast<uint32>(Keys.GetMaxIndex());
+	if (MaxValues < MaxKeys)
 	{
-		if (Entry.BindId == BindId)
+		FCustomBindingHandle Old = { Values, MaxValues };
+		SIZE_T Size = sizeof(ICustomBinding*) + sizeof(FStructDeclaration*) + sizeof(FInnersHandle);
+
+		MaxValues = FMath::RoundUpToPowerOfTwo(FMath::Max(MaxKeys, 4u));
+		Values = reinterpret_cast<ICustomBinding**>(FMemory::MallocZeroed(MaxValues * Size));
+
+		if (Old)
 		{
-			Entries.RemoveAtSwap(&Entry - Entries.GetData(), EAllowShrinking::No);
-			return;
+			FCustomBindingHandle New = { Values, MaxValues };
+			FMemory::Memcpy(New.Bindings,			Old.Bindings,			Old.Max * sizeof(ICustomBinding*));
+			FMemory::Memcpy(New.GetDeclarations(),	Old.GetDeclarations(),	Old.Max * sizeof(FStructDeclaration*));
+			FMemory::Memcpy(New.GetLoweredInners(), Old.GetLoweredInners(), Old.Max * sizeof(FInnersHandle));
+			FMemory::Free(Old.Bindings);
 		}
 	}
-	
-	checkf(false, TEXT("'%s' unbound"), *Debug.Print(BindId));
-}
 
-FOptionalStructSchemaId	FCustomBindings::FindStructDeclId(FStructSchemaId BindId) const
-{ 
-	 FCustomBindingEntry Entry = Find(BindId);
-	 return Entry ? ToOptional(Entry.DeclId) : NoId;
+	// Add value
+	FCustomBindingHandle Value(Values + SetIdx, MaxValues);
+	*Value.Bindings = &Binding;
+	*Value.GetDeclarations() = &Declaration;
+	*Value.GetLoweredInners() = LoweredInners;
+
+	UpdateCulling(/* in-out */ Window, BindId.Idx);
 }
 
-FCustomBindingEntry	FCustomBindings::Find(FStructSchemaId BindId) const
+inline FCustomBindingHandle FCustomBindingMap::Find(FStructSchemaId BindId) const
 {
-	for (FCustomBindingEntry Entry : Entries)
+	if (Cull(Window, BindId.Idx))
 	{
-		if (Entry.BindId == BindId)
+		if (FSetElementId Idx = Keys.FindId(BindId); Idx.IsValidId())
 		{
-			return Entry;
+			return {Values + Idx.AsInteger(), MaxValues};
 		}
 	}
-
-	return Base ? Base->Find(BindId) : FCustomBindingEntry{};
+	return {};
 }
 
-const FCustomBindings& FCustomBindings::GetBottom() const
+void FCustomBindingMap::Drop(FStructSchemaId BindId)
 {
-	const FCustomBindings* It;
-	for (It = this; It->Base; It = It->Base) {}
-	return *It;
+	int32 NumRemoved = Keys.Remove(BindId);
+	checkf(NumRemoved == 1, TEXT("'%s' unbound"), *Debug.Print(BindId));
 }
 
-FInnersHandle FCustomBindings::StoreInners(TConstArrayView<FInnerStruct> In)
+////////////////////////////////////////////////////////////////////////////////////////////////
+
+static FInnersHandle StoreInners(TArray<FInnerStruct>& All, TConstArrayView<FInnerStruct> In)
 {
 	if (uint32 Num = static_cast<uint32>(In.Num()))
 	{
 		// Optimizable: Inline encode if Num == 1
-		TArray<FInnerStruct>& All = GetBottom().AllLoweredInners;
+		// Note: Grows indefinitely, might needs improvement if type erasure becomes common
 		FInnersHandle Out = { Num, static_cast<uint32>(All.Num()) };
 		All.Append(In);
 		return Out;
@@ -228,16 +258,71 @@ FInnersHandle FCustomBindings::StoreInners(TConstArrayView<FInnerStruct> In)
 	return {};
 }
 
-TConstArrayView<FInnerStruct> FCustomBindings::FetchInners(FInnersHandle In) const
+static TConstArrayView<FInnerStruct> FetchInners(const TArray<FInnerStruct>& All, FInnersHandle In)
 {
 	if (int32 Num = static_cast<int32>(In.Num))
 	{
-		const TArray<FInnerStruct>& All = GetBottom().AllLoweredInners;
 		return { &All[In.Idx], Num };
 	}
 	return {};
 }
 
+static bool operator==(FInnerStruct A, FInnerStruct B) { return A.Name == B.Name && A.BindId == B.BindId; }
+
+////////////////////////////////////////////////////////////////////////////////////////////////
+
+void FCustomBindings::BindStruct(FStructSchemaId BindId, ICustomBinding& Binding, const FStructDeclaration& Declaration, TConstArrayView<FInnerStruct> LoweredInners)
+{
+	FInnersHandle StoredInners = StoreInners(/* out */ BottomInners, LoweredInners);
+	check(Algo::Compare(LoweredInners, FetchInners(BottomInners, StoredInners)));
+	Map.Bind(BindId, Binding, Declaration, StoredInners);
+}
+
+const ICustomBinding* FCustomBindings::FindStruct(FStructSchemaId BindId) const
+{
+	FCustomBindingHandle Handle = Find(BindId);
+	return Handle ? *Handle.Bindings : nullptr;
+}
+
+const ICustomBinding* FCustomBindings::FindStruct(FStructSchemaId BindId, TConstArrayView<FInnerStruct>& OutLoweredInners) const
+{
+	if (FCustomBindingHandle Handle = Find(BindId))
+	{
+		OutLoweredInners = FetchInners(BottomInners, *Handle.GetLoweredInners());
+		return *Handle.Bindings;
+	}
+	return nullptr;
+}
+
+ICustomBinding* FCustomBindings::FindStructToSave(FStructSchemaId BindId, const FStructDeclaration*& OutDeclaration) const
+{
+	FCustomBindingHandle Handle = Find(BindId);
+	OutDeclaration = Handle ? *Handle.GetDeclarations() : nullptr;
+	return Handle ? *Handle.Bindings : nullptr;
+}
+
+const FStructDeclaration* FCustomBindings::FindDeclaration(FStructSchemaId BindId) const
+{ 
+	FCustomBindingHandle Handle = Find(BindId);
+	return Handle ? *Handle.GetDeclarations() : nullptr;
+}
+
+void FCustomBindings::DropStruct(FStructSchemaId BindId)
+{
+	Map.Drop(BindId);
+}
+
+FCustomBindingHandle FCustomBindingsBottom::Find(FStructSchemaId BindId) const
+{
+	return Map.Find(BindId);
+}
+
+FCustomBindingHandle FCustomBindingsOverlay::Find(FStructSchemaId BindId) const
+{
+	FCustomBindingHandle Handle = Map.Find(BindId);
+	return Handle ? Handle : Underlay.Find(BindId);
+}
+
 ////////////////////////////////////////////////////////////////////////////////////////////////
 
 inline FTypeId IndexInnermost(FIdIndexerBase& Ids, FMemberBindType Type, FOptionalSchemaId Schema)
@@ -351,8 +436,8 @@ FStructSchemaId FSchemaBindings::GetDeclId(FStructSchemaId BindId) const
 
 FStructSchemaId FStructBindIds::GetDeclId(FStructSchemaId BindId) const
 {
-	FOptionalStructSchemaId CustomId = Customs.FindStructDeclId(BindId);
-	return CustomId ? CustomId.Get() : Schemas.GetDeclId(BindId);
+	const FStructDeclaration* CustomDecl = Customs.FindDeclaration(BindId);
+	return CustomDecl ? CustomDecl->Id : Schemas.GetDeclId(BindId);
 }
 
 //////////////////////////////////////////////////////////////////////////
diff --git a/Engine/Plugins/Experimental/PlainProps/Source/Private/PlainPropsBuild.cpp b/Engine/Plugins/Experimental/PlainProps/Source/Private/PlainPropsBuild.cpp
--- a/Engine/Plugins/Experimental/PlainProps/Source/Private/PlainPropsBuild.cpp
+++ b/Engine/Plugins/Experimental/PlainProps/Source/Private/PlainPropsBuild.cpp
@@ -136,22 +136,25 @@ void FMemberBuilder::BuildSuperStruct(FScratchAllocator& Scratch, const FStructD
 FBuiltStruct* FMemberBuilder::BuildAndReset(FScratchAllocator& Scratch, const FStructDeclaration& Declared, const FDebugIds& Debug)
 {
 	checkf(!(Declared.Super && Declared.Occupancy == EMemberPresence::RequireAll),
-		TEXT("Requiring sub structs to be dense isn't implemented"));
+		TEXT("Requiring sub structs to be dense isn't implemented, %s"), *Debug.Print(Declared.Id));
 #if DO_CHECK
 	// Verify members were added in declared order
 	if (int32 Num = Members.Num())
 	{
-		int32 OrderIdx = 0;
-		TConstArrayView<FMemberId> Order = Declared.GetMemberOrder();
-		int32 SkipSuper = Declared.Super && IsSuper(Members[0].Schema.Type);
-		for (FBuiltMember* It = Members.GetData() + SkipSuper, *End = Members.GetData() + Num; It != End; ++It)
+		if (Declared.bNamedMembers)
 		{
-			FBuiltMember& Member = *It;
-			for (; OrderIdx < Order.Num() && Order[OrderIdx] != Member.Name; ++OrderIdx)
-			{}	
-			checkf(OrderIdx < Order.Num(), TEXT("Member '%s' in '%s' %s"), *Debug.Print(Member.Name), *Debug.Print(Declared.Type),
-					Order.Contains(Member.Name) ? TEXT("appeared in non-declared order") : TEXT("is undeclared"));
-			++OrderIdx;
+			int32 OrderIdx = 0;
+			TConstArrayView<FMemberId> Order = Declared.GetMemberOrder();
+			int32 SkipSuper = Declared.Super && IsSuper(Members[0].Schema.Type);
+			for (FBuiltMember* It = Members.GetData() + SkipSuper, *End = Members.GetData() + Num; It != End; ++It)
+			{
+				FBuiltMember& Member = *It;
+				for (; OrderIdx < Order.Num() && Order[OrderIdx] != Member.Name; ++OrderIdx)
+				{}	
+				checkf(OrderIdx < Order.Num(), TEXT("Member '%s' in '%s' %s"), *Debug.Print(Member.Name), *Debug.Print(Declared.Type),
+						Order.Contains(Member.Name) ? TEXT("appeared in non-declared order") : TEXT("is undeclared"));
+				++OrderIdx;
+			}
 		}
 	}
 #endif
@@ -193,8 +196,8 @@ FBuiltStruct* FDenseMemberBuilder::BuildHomo(const FStructDeclaration& Declarati
 template<typename IntType, typename FloatType>
 inline IntType CheckFiniteBitCast(FloatType Value)
 {
-	// Reject NaN / INF and ignore negative zero for now
-	checkf(FMath::IsFinite(Value), TEXT("Saving NaN or INF isn't supported"));
+	// Todo: Decide if we should accept, reject or sanitize NaN / INF / -0
+	//checkf(FMath::IsFinite(Value), TEXT("Saving NaN or INF isn't supported"));
 	return BitCast<IntType>(Value);
 }
 
diff --git a/Engine/Plugins/Experimental/PlainProps/Source/Private/PlainPropsDeclare.cpp b/Engine/Plugins/Experimental/PlainProps/Source/Private/PlainPropsDeclare.cpp
--- a/Engine/Plugins/Experimental/PlainProps/Source/Private/PlainPropsDeclare.cpp
+++ b/Engine/Plugins/Experimental/PlainProps/Source/Private/PlainPropsDeclare.cpp
@@ -97,7 +97,7 @@ void CopyItems(T* It, TConstArrayView<T> Items)
 	}
 }
 
-static void Declare(TArray<TUniquePtr<FStructDeclaration>>& Out, FStructSchemaId DeclId, FTypeId Type, TConstArrayView<FMemberId> MemberOrder, EMemberPresence Occupancy, FOptionalStructSchemaId Super, uint16 NumMembers)
+static const FStructDeclaration& Declare(TArray<TUniquePtr<FStructDeclaration>>& Out, FStructSchemaId DeclId, FTypeId Type, TConstArrayView<FMemberId> MemberOrder, EMemberPresence Occupancy, FOptionalStructSchemaId Super, uint16 NumMembers)
 {
 	if (static_cast<int32>(DeclId.Idx) >= Out.Num())
 	{
@@ -123,20 +123,22 @@ static void Declare(TArray<TUniquePtr<FStructDeclaration>>& Out, FStructSchemaId
 		Ptr.Reset(new (Data) FStructDeclaration(Header));
 		CopyItems(Ptr->MemberOrder, MemberOrder);
 	}
+	
+	return *Ptr;
 }
 
 
-void FDeclarations::DeclareStruct(FStructSchemaId DeclId, FTypeId Type, TConstArrayView<FMemberId> MemberOrder, EMemberPresence Occupancy, FOptionalStructSchemaId Super)
+const FStructDeclaration& FDeclarations::DeclareStruct(FStructSchemaId DeclId, FTypeId Type, TConstArrayView<FMemberId> MemberOrder, EMemberPresence Occupancy, FOptionalStructSchemaId Super)
 {
-	Declare(/* out */ DeclaredStructs, DeclId, Type, MemberOrder, Occupancy, Super, IntCastChecked<uint16>(MemberOrder.Num()));
+	return Declare(/* out */ DeclaredStructs, DeclId, Type, MemberOrder, Occupancy, Super, IntCastChecked<uint16>(MemberOrder.Num()));
 }
 
-void FDeclarations::DeclareNumeralStruct(FStructSchemaId DeclId, FTypeId Type, uint16 NumMembers, EMemberPresence Occupancy)
+const FStructDeclaration& FDeclarations::DeclareNumeralStruct(FStructSchemaId DeclId, FTypeId Type, uint16 NumMembers, EMemberPresence Occupancy)
 {
-	Declare(/* out */ DeclaredStructs, DeclId, Type, {}, Occupancy, NoId, NumMembers);
+	return Declare(/* out */ DeclaredStructs, DeclId, Type, {}, Occupancy, NoId, NumMembers);
 }
 
-void FDeclarations::DeclareEnum(FEnumSchemaId Id, FTypeId Type, EEnumMode Mode, TConstArrayView<FEnumerator> Enumerators, EEnumAliases Policy)
+const FEnumDeclaration& FDeclarations::DeclareEnum(FEnumSchemaId Id, FTypeId Type, EEnumMode Mode, TConstArrayView<FEnumerator> Enumerators, EEnumAliases Policy)
 {
 	if (static_cast<int32>(Id.Idx) >= DeclaredEnums.Num())
 	{
@@ -158,6 +160,8 @@ void FDeclarations::DeclareEnum(FEnumSchemaId Id, FTypeId Type, EEnumMode Mode,
 	CopyItems(Ptr->Enumerators, Enumerators);
 
 	ValidateDeclaration(*Ptr);
+
+	return *Ptr;
 }
 
 void FDeclarations::DropStructRef(FStructSchemaId DeclId)
diff --git a/Engine/Plugins/Experimental/PlainProps/Source/Private/PlainPropsSave.cpp b/Engine/Plugins/Experimental/PlainProps/Source/Private/PlainPropsSave.cpp
--- a/Engine/Plugins/Experimental/PlainProps/Source/Private/PlainPropsSave.cpp
+++ b/Engine/Plugins/Experimental/PlainProps/Source/Private/PlainPropsSave.cpp
@@ -200,7 +200,7 @@ FBuiltRange* SaveLeafRange(const void* Range, const ILeafRangeBinding& Binding,
 
 [[nodiscard]] static FBuiltRange* SaveStructRange(const void* Range, const IItemRangeBinding& ItemBinding, const FSaveContext& Ctx, FStructSchemaId Id)
 {
-	if (Ctx.Defaults)
+	if (Ctx.Defaults && Ctx.Declarations.Get(Id).Occupancy == EMemberPresence::AllowSparse)
 	{
 		FDefaultStruct Default = { Id, Ctx.Defaults->Get(Id) };
 		return SaveNonLeafRange<FInternalStructRangeDeltaSaver>(Range, ItemBinding, Ctx, Default);
@@ -312,10 +312,9 @@ FBuiltStruct* SaveStruct(const void* Struct, FStructSchemaId BindId, const FSave
 {
 	const FStructDeclaration* Declaration = nullptr;
 	FMemberBuilder Out;
-	if (FCustomBindingEntry Custom = Ctx.Customs.FindStructToSave(BindId))
+	if (ICustomBinding* Custom = Ctx.Customs.FindStructToSave(BindId, /* out */ Declaration))
 	{
-		Custom.Binding->SaveCustom(Out, Struct, nullptr, Ctx);
-		Declaration = &Ctx.Declarations.Get(Custom.DeclId);	
+		Custom->SaveCustom(Out, Struct, nullptr, Ctx);
 	}
 	else
 	{
@@ -440,13 +439,15 @@ bool DiffItemRange(const void* RangeA, const void* RangeB, const IItemRangeBindi
 			{
 				return true;
 			}
-			else if (A.Ctx.Request.NumRead >= NumTotal)
+			else if (A.Ctx.Request.NumRead + Num >= NumTotal)
 			{
-				check(A.Ctx.Request.NumRead == NumTotal);	
-				check(B.Ctx.Request.NumRead == NumTotal);	
+				check(A.Ctx.Request.NumRead + Num == NumTotal);	
+				check(B.Ctx.Request.NumRead + Num == NumTotal);	
 				return false;
 			}
 			
+			A.Ctx.Request.NumRead += Num;
+			B.Ctx.Request.NumRead += Num;
 			A.RefillItems(Binding);
 			B.RefillItems(Binding);
 		}
@@ -529,23 +530,12 @@ static void SaveMemberDelta(FMemberBuilder& Out, const void* Struct, const void*
 	}
 }
 
-FBuiltStruct* SaveStructDelta(const void* Struct, const void* Default, FStructSchemaId BindId, const FSaveContext& Ctx)
+static const FStructDeclaration& SaveSchemaBoundStructDelta(FMemberBuilder& Out, const void* Struct, const void* Default, const FSchemaBinding& Schema, const FSaveContext& Ctx)
 {
-	const FStructDeclaration* Declaration = nullptr;
-	FMemberBuilder Out;
-	if (FCustomBindingEntry Custom = Ctx.Customs.FindStructToSave(BindId))
+	const FStructDeclaration& Declaration = Ctx.Declarations.Get(Schema.DeclId);
+	if (Declaration.bNamedMembers)
 	{
-		Declaration = &Ctx.Declarations.Get(Custom.DeclId);	
-		if (Custom.Binding->DiffCustom(Struct, Default))
-		{
-			Custom.Binding->SaveCustom(Out, Struct, Default, Ctx);
-		}
-	}
-	else
-	{
-		const FSchemaBinding& Schema = Ctx.Schemas.GetStruct(BindId);
-		Declaration = &Ctx.Declarations.Get(Schema.DeclId);	
-		TConstArrayView<FMemberId> MemberOrder = Declaration->GetMemberOrder();
+		TConstArrayView<FMemberId> MemberOrder = Declaration.GetMemberOrder();
 
 		for (FMemberVisitor It(Schema); It.HasMore(); )
 		{
@@ -558,39 +548,62 @@ FBuiltStruct* SaveStructDelta(const void* Struct, const void* Default, FStructSc
 			}
 		}
 	}
-	
+	else // numeral struct
+	{
+		for (FMemberVisitor It(Schema); It.HasMore(); )
+		{
+			// Todo: Make real numeral
+			FMemberId Dummy = { 0 };
+			switch (It.PeekKind())
+			{
+				case EMemberKind::Leaf:		SaveMemberDelta(Out, Struct, Default, Dummy, Ctx, It.GrabLeaf());	break;
+				case EMemberKind::Range:	SaveMemberDelta(Out, Struct, Default, Dummy, Ctx, It.GrabRange());	break;
+				case EMemberKind::Struct:	SaveMemberDelta(Out, Struct, Default, Dummy, Ctx, It.GrabStruct());	break;
+			}
+			++Dummy.Id.Idx;
+		}
+	}
+
+	return Declaration;
+}
+
+FBuiltStruct* SaveStructDelta(const void* Struct, const void* Default, FStructSchemaId BindId, const FSaveContext& Ctx)
+{
+	FMemberBuilder Out;
+	const FStructDeclaration* Declaration = nullptr;
+	if (ICustomBinding* Custom = Ctx.Customs.FindStructToSave(BindId, /* out */ Declaration))
+	{
+		if (Custom->DiffCustom(Struct, Default))
+		{
+			Custom->SaveCustom(Out, Struct, Default, Ctx);
+		}
+	}
+	else
+	{
+		const FSchemaBinding& Schema = Ctx.Schemas.GetStruct(BindId);
+		Declaration = &SaveSchemaBoundStructDelta(Out, Struct, Default, Schema, Ctx);
+	}
 	return Out.BuildAndReset(Ctx.Scratch, *Declaration, Ctx.Declarations.GetDebug());
 }
 
 FBuiltStruct* SaveStructDeltaIfDiff(const void* Struct, const void* Default, FStructSchemaId BindId, const FSaveContext& Ctx)
 {
 	FMemberBuilder Out;
-	if (FCustomBindingEntry Custom = Ctx.Customs.FindStructToSave(BindId))
+	const FStructDeclaration* Declaration = nullptr;
+	if (ICustomBinding* Custom = Ctx.Customs.FindStructToSave(BindId, /* out */ Declaration))
 	{
-		if (Custom.Binding->DiffCustom(Struct, Default))
+		if (Custom->DiffCustom(Struct, Default))
 		{
-			const FStructDeclaration* Declaration = &Ctx.Declarations.Get(Custom.DeclId);	
-			Custom.Binding->SaveCustom(Out, Struct, Default, Ctx);
+			Custom->SaveCustom(Out, Struct, Default, Ctx);
 			return Out.BuildAndReset(Ctx.Scratch, *Declaration, Ctx.Declarations.GetDebug());
 		}
+		
 		return nullptr;
 	}
 
 	const FSchemaBinding& Schema = Ctx.Schemas.GetStruct(BindId);
-	const FStructDeclaration* Declaration = &Ctx.Declarations.Get(Schema.DeclId);	
-	TConstArrayView<FMemberId> MemberOrder = Declaration->GetMemberOrder();
+	Declaration = &SaveSchemaBoundStructDelta(Out, Struct, Default, Schema, Ctx);
 
-	for (FMemberVisitor It(Schema); It.HasMore(); )
-	{
-		FMemberId Name = MemberOrder[It.GetIndex()];
-		switch (It.PeekKind())
-		{
-			case EMemberKind::Leaf:		SaveMemberDelta(Out, Struct, Default, Name, Ctx, It.GrabLeaf());	break;
-			case EMemberKind::Range:	SaveMemberDelta(Out, Struct, Default, Name, Ctx, It.GrabRange());	break;
-			case EMemberKind::Struct:	SaveMemberDelta(Out, Struct, Default, Name, Ctx, It.GrabStruct());	break;
-		}
-	}
-		
 	return Out.IsEmpty() ? nullptr : Out.BuildAndReset(Ctx.Scratch, *Declaration, Ctx.Declarations.GetDebug());
 }
 
diff --git a/Engine/Plugins/Experimental/PlainProps/Source/Private/PlainPropsUeCoreBindings.cpp b/Engine/Plugins/Experimental/PlainProps/Source/Private/PlainPropsUeCoreBindings.cpp
--- a/Engine/Plugins/Experimental/PlainProps/Source/Private/PlainPropsUeCoreBindings.cpp
+++ b/Engine/Plugins/Experimental/PlainProps/Source/Private/PlainPropsUeCoreBindings.cpp
@@ -96,6 +96,62 @@ void FTransformBinding::Load(FTransform& Dst, FStructLoadView Src, ECustomLoadMe
 	checkSlow(!Members.HasMore());
 }
 
+//////////////////////////////////////////////////////////////////////////
+
+void FGuidBinding::Save(FMemberBuilder& Dst, const FGuid& Src, const FGuid*, const FSaveContext&) const
+{
+	Dst.AddHex(MemberIds[0], Src.A);
+	Dst.AddHex(MemberIds[1], Src.B);
+	Dst.AddHex(MemberIds[2], Src.C);
+	Dst.AddHex(MemberIds[3], Src.D);
+}
+
+void FGuidBinding::Load(FGuid& Dst, FStructLoadView Src, ECustomLoadMethod Method) const
+{
+	Src.Values.CheckSize(sizeof(FGuid));
+	FMemory::Memcpy(&Dst, Src.Values.Peek(), sizeof(FGuid));
+}
+
+//////////////////////////////////////////////////////////////////////////
+
+void FColorBinding::Save(FMemberBuilder& Dst, const FColor& Src, const FColor*, const FSaveContext&) const
+{
+	Dst.Add(MemberIds[0], Src.B);
+	Dst.Add(MemberIds[1], Src.G);
+	Dst.Add(MemberIds[2], Src.R);
+	Dst.Add(MemberIds[3], Src.A);
+}
+
+void FColorBinding::Load(FColor& Dst, FStructLoadView Src, ECustomLoadMethod Method) const
+{
+	static_assert(offsetof(FColor, B) == 0);
+	static_assert(offsetof(FColor, G) == 1);
+	static_assert(offsetof(FColor, R) == 2);
+	static_assert(offsetof(FColor, A) == 3);
+	Src.Values.CheckSize(sizeof(FColor));
+	FMemory::Memcpy(&Dst, Src.Values.Peek(), sizeof(FColor));
+}
+
+//////////////////////////////////////////////////////////////////////////
+
+void FLinearColorBinding::Save(FMemberBuilder& Dst, const FLinearColor& Src, const FLinearColor*, const FSaveContext&) const
+{
+	Dst.Add(MemberIds[0], Src.R);
+	Dst.Add(MemberIds[1], Src.G);
+	Dst.Add(MemberIds[2], Src.B);
+	Dst.Add(MemberIds[3], Src.A);
+}
+
+void FLinearColorBinding::Load(FLinearColor& Dst, FStructLoadView Src, ECustomLoadMethod Method) const
+{
+	static_assert(offsetof(FLinearColor, R) == 0);
+	static_assert(offsetof(FLinearColor, G) == 4);
+	static_assert(offsetof(FLinearColor, B) == 8);
+	static_assert(offsetof(FLinearColor, A) == 12);
+	Src.Values.CheckSize(sizeof(FLinearColor));
+	FMemory::Memcpy(&Dst, Src.Values.Peek(), sizeof(FLinearColor));
+}
+
 } // namespace PlainProps::UE
 
 //////////////////////////////////////////////////////////////////////////
diff --git a/Engine/Plugins/Experimental/PlainProps/Source/Public/PlainPropsBind.h b/Engine/Plugins/Experimental/PlainProps/Source/Public/PlainPropsBind.h
--- a/Engine/Plugins/Experimental/PlainProps/Source/Public/PlainPropsBind.h
+++ b/Engine/Plugins/Experimental/PlainProps/Source/Public/PlainPropsBind.h
@@ -3,6 +3,7 @@
 #pragma once
 
 #include "Containers/ArrayView.h"
+#include "Containers/Set.h"
 #include "Containers/StringView.h"
 #include "Memory/MemoryFwd.h"
 #include "Memory/MemoryView.h"
@@ -18,6 +19,7 @@ namespace PlainProps
 
 struct FBuiltRange;
 struct FBuiltStruct;
+struct FCustomBindingHandle;
 class FIdIndexerBase;
 struct FLoadBatch;
 class FMemberBuilder;
@@ -248,40 +250,68 @@ struct FInnersHandle
 	uint32			Idx = 0;
 };
 
-struct FCustomBindingEntry
+struct FIdWindow
 {
-	FStructSchemaId BindId;
-	FStructSchemaId DeclId;
-	ICustomBinding* Binding = nullptr;
-	FInnersHandle	LoweredInners;
+	uint32			Min = 0;
+	uint32			Num = 0;
+};
+
+// Slightly optimized map from FStructSchemaId to ICustomBinding/FStructDeclaration/FInnersHandle
+class FCustomBindingMap
+{
+public:
+	UE_NONCOPYABLE(FCustomBindingMap);
+	explicit FCustomBindingMap(const FDebugIds& Dbg) : Debug(Dbg) {}
+	~FCustomBindingMap() { FMemory::Free(Values); }
+
+	void							Bind(FStructSchemaId BindId, ICustomBinding& Binding, const FStructDeclaration& Declaration, FInnersHandle LoweredInners);
+	FCustomBindingHandle			Find(FStructSchemaId BindId) const;
+	void							Drop(FStructSchemaId BindId);
 
-	explicit operator bool() const { return !!Binding; }
+private:
+	FIdWindow						Window;
+	TSet<FStructSchemaId>			Keys;
+	ICustomBinding**				Values = nullptr;
+	uint32							MaxValues = 0;
+public:
+	const FDebugIds&				Debug;
 };
 
 class FCustomBindings
 {
 public:
-	UE_NONCOPYABLE(FCustomBindings);
-	FCustomBindings(const FDebugIds& Dbg, const FCustomBindings* InBase = nullptr) : Base(InBase), Debug(Dbg) {}
-
 	// @param Binding must outlive this or call DropStruct()
-	PLAINPROPS_API void						BindStruct(FStructSchemaId BindId, FStructSchemaId DeclId, TConstArrayView<FInnerStruct> LoweredInners, ICustomBinding& Binding);
-	const ICustomBinding*					FindStruct(FStructSchemaId BindId) const		{ return Find(BindId).Binding; }
-	PLAINPROPS_API const ICustomBinding*	FindStruct(FStructSchemaId BindId, TConstArrayView<FInnerStruct>& OutLoweredInners) const;
-	FOptionalStructSchemaId					FindStructDeclId(FStructSchemaId BindId) const;
-	FCustomBindingEntry						FindStructToSave(FStructSchemaId BindId)		{ return Find(BindId); }
-	PLAINPROPS_API void						DropStruct(FStructSchemaId BindId);
+	PLAINPROPS_API void							BindStruct(FStructSchemaId BindId, ICustomBinding& Binding, const FStructDeclaration& Declaration, TConstArrayView<FInnerStruct> LoweredInners = {});
+	PLAINPROPS_API const ICustomBinding*		FindStruct(FStructSchemaId BindId) const;
+	PLAINPROPS_API const ICustomBinding*		FindStruct(FStructSchemaId BindId, TConstArrayView<FInnerStruct>& OutLoweredInners) const;
+	PLAINPROPS_API ICustomBinding*				FindStructToSave(FStructSchemaId BindId, const FStructDeclaration*& OutDeclaration) const;
+	PLAINPROPS_API const FStructDeclaration*	FindDeclaration(FStructSchemaId BindId) const;
+	PLAINPROPS_API void							DropStruct(FStructSchemaId BindId);
+
+	virtual FCustomBindingHandle				Find(FStructSchemaId BindId) const = 0;
+protected:
+	FCustomBindings(TArray<FInnerStruct>& Inners, const FDebugIds& Dbg) : Map(Dbg), BottomInners(Inners) {}
+	FCustomBindings(const FCustomBindings* Under) : Map(Under->Map.Debug), BottomInners(Under->BottomInners) {}
+	~FCustomBindings() = default;
+
+	FCustomBindingMap							Map;
+	TArray<FInnerStruct>&						BottomInners;
+};
 
-private:
-	PLAINPROPS_API FCustomBindingEntry		Find(FStructSchemaId BindId) const;
-	const FCustomBindings&					GetBottom() const;
-	FInnersHandle							StoreInners(TConstArrayView<FInnerStruct> In);
-	TConstArrayView<FInnerStruct>			FetchInners(FInnersHandle In) const;
-	
-	const FCustomBindings*					Base = nullptr;
-	TArray<FCustomBindingEntry, TInlineAllocator<8>>		Entries;
-	mutable TArray<FInnerStruct>			AllLoweredInners;
-	const FDebugIds&						Debug;
+class FCustomBindingsBottom final : public FCustomBindings
+{
+	mutable TArray<FInnerStruct>				AllInners;
+	PLAINPROPS_API virtual FCustomBindingHandle	Find(FStructSchemaId BindId) const override;	
+public:
+	explicit FCustomBindingsBottom(const FDebugIds& Dbg) : FCustomBindings(/* unconstructed */  AllInners, Dbg) {}
+};
+
+class FCustomBindingsOverlay final : public FCustomBindings
+{
+	const FCustomBindings&						Underlay;
+	PLAINPROPS_API virtual FCustomBindingHandle	Find(FStructSchemaId BindId) const override;
+public:
+	explicit FCustomBindingsOverlay(const FCustomBindings& Under) : FCustomBindings(&Under), Underlay(Under) {}
 };
 
 template<typename T>
@@ -425,17 +455,17 @@ struct FExistingItems
 	uint32				Stride = 0;
 	FExistingItemSlice	Slice;
 
-	void SetAll(FExistingItemSlice Whole, uint32 InStride)
+	void SetAll(const void* Items, uint64 Num, uint32 InStride)
 	{
-		NumTotal = Whole.Num;
+		NumTotal = Num;
 		Stride = InStride;
-		Slice = Whole;
+		Slice = { Items, Num };
 	}
 
 	template<typename ItemType>
-	void SetAll(const ItemType* Items, uint64 NumItems)
+	void SetAll(const ItemType* Items, uint64 Num)
 	{
-		SetAll(FExistingItemSlice{Items, NumItems}, sizeof(ItemType));
+		SetAll(Items, Num, sizeof(ItemType));
 	}
 };
 
@@ -564,7 +594,7 @@ private:
 
 ////////////////////////////////////////////////////////////////////////////////////////////////
 
-// Lookup default instances when saving struct ranges
+// Lookup default instances when delta-saving struct ranges
 struct IDefaultStructs
 {
 	virtual const void* Get(FStructSchemaId Id) = 0;
@@ -587,22 +617,21 @@ struct FStructBindIds : IStructBindIds
 
 ////////////////////////////////////////////////////////////////////////////////////////////////
 
-template<class Ids, typename Typename>
-FScopeId IndexNamespaceId()
+template<class Ids, class Typename>
+FScopeId IndexNamespaceId() requires (Typename::Namespace.size() > 0)
 {
 	// Opt: Make cached GetNamespaceId(), either via new namespace CTTI types (maybe PP_REFLECT_NAMESPACE)
 	//		or some compile time string template parameters, perhaps a variadic template taking any number of chars
-	if constexpr (Typename::Namespace.size())
-	{
-		return Ids::IndexScope(ToAnsiView(Typename::Namespace));
-	}
-	else
-	{
-		return NoId;
-	}
+	return Ids::IndexScope(ToAnsiView(Typename::Namespace));
 }
 
-template<ETypename Kind, typename Typename>
+template<class Ids, class Typename>
+FScopeId IndexNamespaceId()
+{
+	return NoId;
+}
+
+template<ETypename Kind, class Typename>
 constexpr std::string_view SelectStructName()
 {	
 	if constexpr (Kind == ETypename::Bind && ExplicitBindName<Typename>)
@@ -615,11 +644,11 @@ constexpr std::string_view SelectStructName()
 	}
 }
 
-template<class Ids, ETypename Kind, typename Typename>
+template<class Ids, ETypename Kind, class Typename>
 FTypeId IndexStructName()
 {
 	FTypeId BaseName = { IndexNamespaceId<Ids, Typename>(), 
-						 Ids::IndexTypename(ToAnsiView(SelectStructName<Kind, Typename>())) };
+						 FTypenameId(Ids::IndexTypename(ToAnsiView(SelectStructName<Kind, Typename>()))) };
 	
 	if constexpr (ParametricName<Typename>)
 	{
@@ -631,7 +660,7 @@ FTypeId IndexStructName()
 	}
 }
 
-template<class Ids, typename Typename>
+template<class Ids, class Typename>
 FStructSchemaId IndexStructBindIdIfNeeded(FStructSchemaId DeclId)
 {
 	if constexpr (ExplicitBindName<Typename> || ParametricName<Typename>)
@@ -645,7 +674,7 @@ FStructSchemaId IndexStructBindIdIfNeeded(FStructSchemaId DeclId)
 	}
 }
 
-template<class Ids, typename Typename>
+template<class Ids, class Typename>
 FDualStructSchemaId IndexStructDualId(FTypeId DeclName = IndexStructName<Ids, ETypename::Decl, Typename>())
 {	
 	FDualStructSchemaId Out;
@@ -802,11 +831,23 @@ struct TScopedStructBinding : FDualStructSchemaId, CustomBinding
 	using Typename = TCustomTypename<CustomBinding>;
 
 	explicit TScopedStructBinding(FTypeId DeclName = IndexStructName<Ids, ETypename::Decl, Typename>(), InnerStructArray Lowered = {})
-	: FDualStructSchemaId(IndexStructDualId<Ids, Typename>(DeclName))
+	: TScopedStructBinding(DeclName, IndexStructDualId<Ids, Typename>(DeclName), Lowered)
+	{}
+
+	explicit TScopedStructBinding(FDualStructSchemaId Id, InnerStructArray Lowered = {})
+	: TScopedStructBinding(Ids::GetIndexer().Resolve(Id.DeclId), Id, Lowered)
+	{}
+
+	explicit TScopedStructBinding(FStructSchemaId Id, InnerStructArray Lowered = {})
+	: TScopedStructBinding(Ids::GetIndexer().Resolve(Id), {Id, Id}, Lowered)
+	{}
+
+	TScopedStructBinding(FTypeId DeclName, FDualStructSchemaId DualId, InnerStructArray& Lowered)
+	: FDualStructSchemaId(DualId)
 	, CustomBinding(TCustomInit<Ids>{{Lowered}})
 	{
-		Runtime::GetTypes().DeclareStruct(DeclId, DeclName, CustomBinding::MemberIds, TOccupancyOf<T>::Value);
-		Runtime::GetCustoms().BindStruct(BindId, DeclId, Lowered, *this);
+		const FStructDeclaration& Declaration = Runtime::GetTypes().DeclareStruct(DeclId, DeclName, CustomBinding::MemberIds, TOccupancyOf<T>::Value);
+		Runtime::GetCustoms().BindStruct(BindId, *this, Declaration, Lowered);
 	}
 
 	~TScopedStructBinding()
diff --git a/Engine/Plugins/Experimental/PlainProps/Source/Public/PlainPropsBuild.h b/Engine/Plugins/Experimental/PlainProps/Source/Public/PlainPropsBuild.h
--- a/Engine/Plugins/Experimental/PlainProps/Source/Public/PlainPropsBuild.h
+++ b/Engine/Plugins/Experimental/PlainProps/Source/Public/PlainPropsBuild.h
@@ -239,7 +239,7 @@ public:
 	template<Arithmetic T>
 	void Add(FMemberId Name, T Value)
 	{
-		AddLeaf(Name, ReflectArithmetic<T>, {}, ValueCast(Value));
+		AddLeaf(Name, ReflectArithmetic<T>, NoId, ValueCast(Value));
 	}
 
 	template<Enumeration T>
@@ -252,6 +252,11 @@ public:
 	void AddEnum16(FMemberId Name, FEnumSchemaId Schema, uint16 Value)	{ AddLeaf(Name, {ELeafType::Enum, ELeafWidth::B16}, ToOptional(Schema), Value); }
 	void AddEnum32(FMemberId Name, FEnumSchemaId Schema, uint32 Value)	{ AddLeaf(Name, {ELeafType::Enum, ELeafWidth::B32}, ToOptional(Schema), Value); }
 	void AddEnum64(FMemberId Name, FEnumSchemaId Schema, uint64 Value)	{ AddLeaf(Name, {ELeafType::Enum, ELeafWidth::B64}, ToOptional(Schema), Value); }
+	
+	void AddHex(FMemberId Name, uint8 Value)							{ AddLeaf(Name, {ELeafType::Hex, ELeafWidth::B8}, NoId, Value); }
+	void AddHex(FMemberId Name, uint16 Value)							{ AddLeaf(Name, {ELeafType::Hex, ELeafWidth::B16}, NoId, Value); }
+	void AddHex(FMemberId Name, uint32 Value)							{ AddLeaf(Name, {ELeafType::Hex, ELeafWidth::B32}, NoId, Value); }
+	void AddHex(FMemberId Name, uint64 Value)							{ AddLeaf(Name, {ELeafType::Hex, ELeafWidth::B64}, NoId, Value); }
 
 	void AddLeaf(FMemberId Name, FUnpackedLeafType Leaf, FOptionalEnumSchemaId Enum, uint64 Value)	{ Members.Emplace(Name, Leaf, Enum, Value); }
 	void AddRange(FMemberId Name, FTypedRange Range)												{ Members.Emplace(Name, Range); }
diff --git a/Engine/Plugins/Experimental/PlainProps/Source/Public/PlainPropsDeclare.h b/Engine/Plugins/Experimental/PlainProps/Source/Public/PlainPropsDeclare.h
--- a/Engine/Plugins/Experimental/PlainProps/Source/Public/PlainPropsDeclare.h
+++ b/Engine/Plugins/Experimental/PlainProps/Source/Public/PlainPropsDeclare.h
@@ -64,13 +64,13 @@ public:
 	UE_NONCOPYABLE(FDeclarations);
 	explicit FDeclarations(const FDebugIds& In) : Debug(In) {}
 
-	PLAINPROPS_API void								DeclareEnum(FEnumSchemaId DeclId, FTypeId Type, EEnumMode Mode, TConstArrayView<FEnumerator> Enumerators, EEnumAliases Policy);
+	PLAINPROPS_API const FEnumDeclaration&			DeclareEnum(FEnumSchemaId DeclId, FTypeId Type, EEnumMode Mode, TConstArrayView<FEnumerator> Enumerators, EEnumAliases Policy);
 	// Declare struct with ref count 1 or increment it and check that previous declaration matches
-	PLAINPROPS_API void								DeclareStruct(FStructSchemaId DeclId, FTypeId Type, TConstArrayView<FMemberId> MemberOrder, EMemberPresence Occupancy, FOptionalStructSchemaId Super = {});
-	PLAINPROPS_API void								DeclareNumeralStruct(FStructSchemaId DeclId, FTypeId Type, uint16 NumMembers, EMemberPresence Occupancy);
+	PLAINPROPS_API const FStructDeclaration&		DeclareStruct(FStructSchemaId DeclId, FTypeId Type, TConstArrayView<FMemberId> MemberOrder, EMemberPresence Occupancy, FOptionalStructSchemaId Super = {});
+	PLAINPROPS_API const FStructDeclaration&		DeclareNumeralStruct(FStructSchemaId DeclId, FTypeId Type, uint16 NumMembers, EMemberPresence Occupancy);
 	
 	void											DropEnum(FEnumSchemaId DeclId)			{ Check(DeclId); DeclaredEnums[DeclId.Idx].Reset(); }
-	void											DropStructRef(FStructSchemaId DeclId);
+	PLAINPROPS_API void								DropStructRef(FStructSchemaId DeclId);
 
 	const FEnumDeclaration&							Get(FEnumSchemaId DeclId) const			{ Check(DeclId); return *DeclaredEnums[DeclId.Idx]; }
 	const FStructDeclaration&						Get(FStructSchemaId DeclId) const		{ Check(DeclId); return *DeclaredStructs[DeclId.Idx]; }
diff --git a/Engine/Plugins/Experimental/PlainProps/Source/Public/PlainPropsSave.h b/Engine/Plugins/Experimental/PlainProps/Source/Public/PlainPropsSave.h
--- a/Engine/Plugins/Experimental/PlainProps/Source/Public/PlainPropsSave.h
+++ b/Engine/Plugins/Experimental/PlainProps/Source/Public/PlainPropsSave.h
@@ -27,9 +27,9 @@ struct FSaveContext
 };
 
 template<typename Runtime>
-FSaveContext MakeSaveContext(FScratchAllocator& Scratch, IDefaultStructs* Defaults)
+FSaveContext MakeSaveContext(FScratchAllocator& Scratch)
 {
-	return { Runtime::GetTypes(), Runtime::GetSchemas(), Runtime::GetCustoms(), Scratch, Defaults };
+	return { Runtime::GetTypes(), Runtime::GetSchemas(), Runtime::GetCustoms(), Scratch, Runtime::GetDefaults() };
 }
 
 ////////////////////////////////////////////////////////////////////////////////////////////////
diff --git a/Engine/Plugins/Experimental/PlainProps/Source/Public/PlainPropsUeCoreBindings.h b/Engine/Plugins/Experimental/PlainProps/Source/Public/PlainPropsUeCoreBindings.h
--- a/Engine/Plugins/Experimental/PlainProps/Source/Public/PlainPropsUeCoreBindings.h
+++ b/Engine/Plugins/Experimental/PlainProps/Source/Public/PlainPropsUeCoreBindings.h
@@ -486,6 +486,57 @@ struct FTransformBinding : ICustomBinding
 
 //////////////////////////////////////////////////////////////////////////
 
+struct FGuidBinding : ICustomBinding
+{
+	using Type = FGuid;
+	const FMemberId MemberIds[4];
+
+	template<class Ids>
+	FGuidBinding(TCustomInit<Ids>)
+	: MemberIds{Ids::IndexMember("A"), Ids::IndexMember("B"), Ids::IndexMember("C"), Ids::IndexMember("D")}
+	{}
+
+	PLAINPROPS_API void	Save(FMemberBuilder& Dst, const FGuid& Src, const FGuid* Default, const FSaveContext&) const;
+	PLAINPROPS_API void	Load(FGuid& Dst, FStructLoadView Src, ECustomLoadMethod Method) const;
+	inline static bool	Diff(FGuid A, FGuid B) { return A != B; }
+};
+
+//////////////////////////////////////////////////////////////////////////
+
+struct FColorBinding : ICustomBinding
+{
+	using Type = FColor;
+	const FMemberId MemberIds[4];
+
+	template<class Ids>
+	FColorBinding(TCustomInit<Ids>)
+	: MemberIds{Ids::IndexMember("B"), Ids::IndexMember("G"), Ids::IndexMember("R"), Ids::IndexMember("A")}
+	{}
+
+	PLAINPROPS_API void	Save(FMemberBuilder& Dst, const FColor& Src, const FColor* Default, const FSaveContext&) const;
+	PLAINPROPS_API void	Load(FColor& Dst, FStructLoadView Src, ECustomLoadMethod Method) const;
+	inline static bool	Diff(FColor A, FColor B) { return A != B; }
+};
+
+//////////////////////////////////////////////////////////////////////////
+
+struct FLinearColorBinding : ICustomBinding
+{
+	using Type = FLinearColor;
+	const FMemberId MemberIds[4];
+
+	template<class Ids>
+	FLinearColorBinding(TCustomInit<Ids>)
+	: MemberIds{Ids::IndexMember("R"), Ids::IndexMember("G"), Ids::IndexMember("B"), Ids::IndexMember("A")}
+	{}
+
+	PLAINPROPS_API void	Save(FMemberBuilder& Dst, const FLinearColor& Src, const FLinearColor* Default, const FSaveContext&) const;
+	PLAINPROPS_API void	Load(FLinearColor& Dst, FStructLoadView Src, ECustomLoadMethod Method) const;
+	inline static bool	Diff(FLinearColor A, FLinearColor B) { return A != B; }
+};
+
+//////////////////////////////////////////////////////////////////////////
+
 struct FBaseDeltaBinding
 {
 	enum class EOp { Del, Add };
@@ -875,35 +926,20 @@ struct TMapDeltaBinding : TBaseDeltaBinding<TMap<K, V, SetAllocator, KeyFuncs>,
 namespace PlainProps
 {
 
-template <>
-struct TTypename<FName>
-{
-	inline static constexpr std::string_view DeclName = "Name";
-	inline static constexpr std::string_view Namespace;
-};
-
-template <>
-struct TTypename<FTransform>
-{
-	inline static constexpr std::string_view DeclName = "Transform";
-	inline static constexpr std::string_view Namespace;
-};
+template <> struct TTypename<FName>					{ inline static constexpr std::string_view DeclName = "Name"; };
+template <> struct TTypename<FTransform>			{ inline static constexpr std::string_view DeclName = "Transform"; };
+template <> struct TTypename<FGuid>					{ inline static constexpr std::string_view DeclName = "Guid"; };
+template <> struct TTypename<FColor>				{ inline static constexpr std::string_view DeclName = "Color"; };
+template <> struct TTypename<FLinearColor>			{ inline static constexpr std::string_view DeclName = "LinearColor"; };
+template <> struct TTypename<FString>				{ inline static constexpr std::string_view RangeBindName = "String"; };
 
 template <typename K, typename V>
 struct TTypename<TPair<K,V>>
 {
 	inline static constexpr std::string_view DeclName = "Pair";
-	inline static constexpr std::string_view Namespace;
 	using Parameters = std::tuple<K, V>;
 };
 
-template <>
-struct TTypename<FString>
-{
-	inline static constexpr std::string_view RangeBindName = "String";
-	inline static constexpr std::string_view Namespace;
-};
-
 inline static constexpr std::string_view UeArrayName = "Array";
 inline static constexpr std::string_view UeSetName = "Set";
 inline static constexpr std::string_view UeMapName = "Map";
@@ -957,23 +993,9 @@ struct TRangeBind<TArray<T, Allocator>>
 	using Type = UE::TArrayBinding<T, Allocator>;
 };
 
-template<>
-struct TRangeBind<FString>
-{
-	using Type = UE::TStringBinding<FString>;
-};
-
-template<>
-struct TRangeBind<FAnsiString>
-{
-	using Type = UE::TStringBinding<FAnsiString>;
-};
-
-template<>
-struct TRangeBind<FUtf8String>
-{
-	using Type = UE::TStringBinding<FUtf8String>;
-};
+template<> struct TRangeBind<FString> { using Type = UE::TStringBinding<FString>; };
+template<> struct TRangeBind<FAnsiString> {	using Type = UE::TStringBinding<FAnsiString>; };
+template<> struct TRangeBind<FUtf8String> {	using Type = UE::TStringBinding<FUtf8String>; };
 
 template<typename T>
 struct TRangeBind<TUniquePtr<T>>
@@ -1011,17 +1033,20 @@ struct TRangeBind<TOptional<T>>
 	using Type = UE::TOptionalBinding<T>;
 };
 
-template<>
-struct TOccupancyOf<FQuat> : FRequireAll {};
+template<> struct TOccupancyOf<FQuat> : FRequireAll {};
+template<> struct TOccupancyOf<FVector> : FRequireAll {};
+template<> struct TOccupancyOf<FGuid> : FRequireAll {};
+template<> struct TOccupancyOf<FColor> : FRequireAll {};
+template<> struct TOccupancyOf<FLinearColor> : FRequireAll {};
 
-template<>
-struct TOccupancyOf<FVector> : FRequireAll {};
+template<> struct TCustomBind<FTransform> {	using Type = UE::FTransformBinding; };
+template<> struct TCustomBind<FGuid> { using Type = UE::FGuidBinding; };
+template<> struct TCustomBind<FColor> { using Type = UE::FColorBinding; };
+template<> struct TCustomBind<FLinearColor> { using Type = UE::FLinearColorBinding; };
 
-template<>
-struct TCustomBind<FTransform>
-{
-	using Type = UE::FTransformBinding;
-};
+
+template<typename K, typename V>
+struct TOccupancyOf<TPair<K, V>> : FRequireAll {};
 
 template<typename K, typename V> requires (!std::is_default_constructible_v<TPair<K,V>>)
 struct TCustomBind<TPair<K, V>>
@@ -1029,7 +1054,4 @@ struct TCustomBind<TPair<K, V>>
 	using Type = UE::TPairBinding<K, V>;
 };
 
-template<typename K, typename V>
-struct TOccupancyOf<TPair<K, V>> : FRequireAll {};
-
 } // namespace PlainProps
diff --git a/Engine/Plugins/Experimental/PlainProps/Source/Tests/PlainPropsUeCoreTests.cpp b/Engine/Plugins/Experimental/PlainProps/Source/Tests/PlainPropsUeCoreTests.cpp
--- a/Engine/Plugins/Experimental/PlainProps/Source/Tests/PlainPropsUeCoreTests.cpp
+++ b/Engine/Plugins/Experimental/PlainProps/Source/Tests/PlainPropsUeCoreTests.cpp
@@ -28,11 +28,11 @@ DEFINE_LOG_CATEGORY_STATIC(LogPlainPropsUeCoreTests, Log, All);
 namespace PlainProps::UE::Test
 {
 
-static TIdIndexer<FName>	GNames;
-static FDeclarations		GTypes(/* debug */ GNames);
-static FSchemaBindings		GSchemas(/* debug */ GNames);
-static FCustomBindings		GCustoms(/* debug */ GNames);
-static FCustomBindings		GDeltaCustoms(/* debug */ GNames, /* base*/ &GCustoms);
+static TIdIndexer<FName>		GNames;
+static FDeclarations			GTypes(/* debug */ GNames);
+static FSchemaBindings			GSchemas(/* debug */ GNames);
+static FCustomBindingsBottom	GCustoms(/* debug */ GNames);
+static FCustomBindingsOverlay	GDeltaCustoms(/* base*/ GCustoms);
 
 struct FIds
 {
@@ -154,14 +154,14 @@ private:
 	using IdBuiltStructPair = TPair<FStructSchemaId, FBuiltStruct*>;
 	TArray<IdBuiltStructPair>	SavedObjects;
 	FNameBinding				SavedNames;
-	FCustomBindings				Customs;
+	FCustomBindingsOverlay		Customs;
 	mutable FScratchAllocator	Scratch;
 };
 
-FBatchSaver::FBatchSaver(const FCustomBindings& CustomBase)
-: Customs(/* debug */ GNames, &CustomBase)
+FBatchSaver::FBatchSaver(const FCustomBindings& Underlay)
+: Customs(Underlay)
 {
-	Customs.BindStruct(SavedNames.Declaration.Id, SavedNames.Declaration.Id, {}, SavedNames);
+	Customs.BindStruct(SavedNames.Declaration.Id, SavedNames, GTypes.Get(SavedNames.Declaration.Id), {});
 }
 
 template<class T>
@@ -330,8 +330,8 @@ void PrintLoadStruct(const FIdBinding& Binding, FStructView StructView, int32 Lo
 class FBatchLoader
 {
 public:
-	FBatchLoader(FMemoryView Data, const FCustomBindings& CustomBase)
-	: Customs(/* debug */ GNames, &CustomBase)
+	FBatchLoader(FMemoryView Data, const FCustomBindings& Underlay)
+	: Customs(Underlay)
 	{
 		// Read ids
 		FByteReader It(Data);
@@ -371,7 +371,7 @@ public:
 
 		// Read names and bind custom loader
 		Names.Names.Append(GrabNumAndArray<FName>(It));
-		Customs.BindStruct(Names.Declaration.Id, Names.Declaration.Id, {}, Names);
+		Customs.BindStruct(Names.Declaration.Id, Names, GTypes.Get(Names.Declaration.Id), {});
 		CHECK(It.Grab<uint32>() == Magics[5]);
 
 		// Finally create load plans
@@ -416,7 +416,7 @@ private:
 	TConstArrayView<FName>		Ids;
 	FNameBinding				Names;
 	FIdBinding					IdBinding;
-	FCustomBindings				Customs;
+	FCustomBindingsOverlay		Customs;
 	FLoadBatchPtr				Plans;
 	TArray<FStructView>			Objects;
 	int32						LoadIdx = 0;
diff --git a/Engine/Plugins/Experimental/PlainPropsUObject/Source/Private/PlainPropsUObjectBindings.cpp b/Engine/Plugins/Experimental/PlainPropsUObject/Source/Private/PlainPropsUObjectBindings.cpp
--- a/Engine/Plugins/Experimental/PlainPropsUObject/Source/Private/PlainPropsUObjectBindings.cpp
+++ b/Engine/Plugins/Experimental/PlainPropsUObject/Source/Private/PlainPropsUObjectBindings.cpp
@@ -2,6 +2,8 @@
 
 #include "PlainPropsUObjectRuntime.h"
 #include "PlainPropsUeCoreBindings.h"
+#include "PlainPropsBuildSchema.h"
+#include "PlainPropsWrite.h"
 #include "Algo/Compare.h"
 #include "Containers/PagedArray.h"
 #include "Hash/xxhash.h"
@@ -13,6 +15,7 @@
 #include "UObject/Class.h"
 #include "UObject/EnumProperty.h"
 #include "UObject/FieldPathProperty.h"
+#include "UObject/Object.h"
 #include "UObject/PropertyOptional.h"
 #include "UObject/StrProperty.h"
 #include "UObject/TextProperty.h"
@@ -22,16 +25,38 @@
 #include "UObject/VerseValueProperty.h"
 
 
+using FUnicastScriptDelegate = TScriptDelegate<FNotThreadSafeNotCheckedDelegateMode>;
+using FMulticastInvocationList = TArray<FUnicastScriptDelegate>;
+using FMulticastInvocationView = TConstArrayView<FUnicastScriptDelegate>;
+
+// Temp hacks. Long-term either add FProperty getters for ctor/dtor/hash function pointers 
+// and delegate APIs for non-intrusive serialization or integrate PlainProps into Core/CoreUObject
+UE_DEFINE_PRIVATE_MEMBER_PTR(void(void*) const, GInitPropertyValue, FProperty, InitializeValueInternal);
+UE_DEFINE_PRIVATE_MEMBER_PTR(void(void*) const, GDestroyPropertyValue, FProperty, DestroyValueInternal);
+UE_DEFINE_PRIVATE_MEMBER_PTR(uint32(const void*) const, GHashPropertyValue, FProperty, GetValueTypeHashInternal);
+UE_DEFINE_PRIVATE_MEMBER_PTR(TArray<FName>, GFieldPathPath, FFieldPath, Path);
+UE_DEFINE_PRIVATE_MEMBER_PTR(TWeakObjectPtr<UStruct>, GFieldPathOwner, FFieldPath, ResolvedOwner);
+UE_DEFINE_PRIVATE_MEMBER_PTR(FWeakObjectPtr, GDelegateObject, FScriptDelegate, Object);
+UE_DEFINE_PRIVATE_MEMBER_PTR(FName, GDelegateFunctionName, FScriptDelegate, FunctionName);
+UE_DEFINE_PRIVATE_MEMBER_PTR(FWeakObjectPtr, GUnicastDelegateObject, FUnicastScriptDelegate, Object);
+UE_DEFINE_PRIVATE_MEMBER_PTR(FName, GUnicastDelegateFunctionName, FUnicastScriptDelegate, FunctionName);
+UE_DEFINE_PRIVATE_MEMBER_PTR(FMulticastInvocationList, GMulticastDelegateInvocationList, FMulticastScriptDelegate, InvocationList);
+UE_DEFINE_PRIVATE_MEMBER_PTR(bool, GSparseDelegateIsBound, FSparseDelegate, bIsBound);
+
 namespace PlainProps::UE
 {
 
+static constexpr ERangeSizeType DefaultRangeMax = RangeSizeOf(FDefaultAllocator::SizeType{});
+
+////////////////////////////////////////////////////////////////////////////////////////////////
+
 struct FDefaultStruct
 {
 	UScriptStruct::ICppStructOps&		Ops;
 	alignas(16) uint8					Instance[0];
 };
 
-static FDefaultStruct* NewDefault(UScriptStruct::ICppStructOps& Ops)
+static FDefaultStruct* NewDefaultStruct(UScriptStruct::ICppStructOps& Ops)
 {
 	check(Ops.GetAlignment() <= 16);
 	uint32 Size = sizeof(FDefaultStruct) + Ops.GetSize();
@@ -41,7 +66,7 @@ static FDefaultStruct* NewDefault(UScriptStruct::ICppStructOps& Ops)
 	return Out;
 }
 
-static void DeleteDefault(uint8* Instance)
+inline void	DeleteDefaultStruct(uint8* Instance)
 {
 	FDefaultStruct* Struct = reinterpret_cast<FDefaultStruct*>(Instance - offsetof(FDefaultStruct, Instance));
 	if (Struct->Ops.HasDestructor())
@@ -51,6 +76,18 @@ static void DeleteDefault(uint8* Instance)
 	FMemory::Free(Struct);
 }
 
+static constexpr uint64			DefaultInstanceStaticMask = 1;
+inline FDefaultInstance			MakeStaticInstance(const void* Static)			{ return { reinterpret_cast<uint64>(Static) | DefaultInstanceStaticMask }; }
+inline FDefaultInstance			MakeDefaultInstance(FDefaultStruct* Default)	{ return { reinterpret_cast<uint64>(Default->Instance) }; }
+inline uint8*					GetInstance(FDefaultInstance Instance)			{ return reinterpret_cast<uint8*>(Instance.Ptr & ~DefaultInstanceStaticMask); }
+inline void						DeleteInstance(FDefaultInstance Instance)
+{
+	if (!(Instance.Ptr & DefaultInstanceStaticMask))
+	{
+		DeleteDefaultStruct(reinterpret_cast<uint8*>(Instance.Ptr));
+	}
+}
+
 static void ReserveZeroes(/* in-out */ FMutableMemoryView& Zeroes, SIZE_T Size, uint32 Alignment)
 {
 	Size +=  FMath::Max<int32>(0, Alignment - 16);
@@ -64,9 +101,26 @@ static void ReserveZeroes(/* in-out */ FMutableMemoryView& Zeroes, SIZE_T Size,
 
 FDefaultStructs::~FDefaultStructs()
 {
-	for (TPair<FStructSchemaId, uint8*> Instance : Instances)
+	for (TPair<FStructSchemaId, FDefaultInstance> Instance : Instances)
 	{
-		DeleteDefault(Instance.Value);
+		DeleteInstance(Instance.Value);
+	}
+}
+
+inline bool Flip(FBitReference Bit)
+{
+	Bit = !Bit;
+	return Bit;
+}
+
+void FDefaultStructs::ReserveFlags(uint32 Idx)
+{
+	if (Idx >= static_cast<uint32>(Instanced.Num()))
+	{
+		Instanced.SetNum(FMath::RoundUpToPowerOfTwo64(Idx + 1), false);
+#if DO_CHECK
+		Bound.SetNum(Instanced.Num(), false);
+#endif
 	}
 }
 
@@ -75,13 +129,11 @@ void FDefaultStructs::Bind(FStructSchemaId Id, const UScriptStruct* Struct)
 	const EStructFlags Flags = Struct->StructFlags;
 	const SIZE_T Size = Struct->GetStructureSize();
 	const uint32 Alignment = Struct->GetMinAlignment();
+	UScriptStruct::ICppStructOps* Ops = Struct->GetCppStructOps();
 
-	if (Id.Idx >= static_cast<uint32>(Instanced.Num()))
-	{
-		Instanced.SetNum(FMath::RoundUpToPowerOfTwo64(Id.Idx + 1), false);
-	}
+	ReserveFlags(Id.Idx);
+	checkf(Flip(Bound[Id.Idx]), TEXT("'%s' already bound"), *GUE.Names.Print(Id));
 
-	UScriptStruct::ICppStructOps* Ops = Struct->GetCppStructOps();
 	if (!!(Flags & STRUCT_ZeroConstructor) || Ops == nullptr)
 	{
 		ReserveZeroes(Zeroes, Size, Alignment);
@@ -89,56 +141,86 @@ void FDefaultStructs::Bind(FStructSchemaId Id, const UScriptStruct* Struct)
 	else
 	{
 		check(Ops->GetSize() == Size);
-		FDefaultStruct* Default = NewDefault(*Ops);
+		FDefaultStruct* Default = NewDefaultStruct(*Ops);
 		if (FMemory::MemIsZero(Default->Instance, Size))
 		{
-			DeleteDefault(Default->Instance);
+			DeleteDefaultStruct(Default->Instance);
 			ReserveZeroes(Zeroes, Size, Alignment);
 		}
 		else
 		{
 			Instanced[Id.Idx] = true;
-			Instances.Add(Id, Default->Instance);
+			Instances.Add(Id, MakeDefaultInstance(Default));
 		}
 	}
 }
 
+void FDefaultStructs::BindZeroes(FStructSchemaId Id, SIZE_T Size, uint32 Alignment)
+{
+	ReserveFlags(Id.Idx);
+	checkf(Flip(Bound[Id.Idx]), TEXT("'%s' already bound"), *GUE.Names.Print(Id));
+	ReserveZeroes(Zeroes, Size, Alignment);
+}
+
+void FDefaultStructs::BindStatic(FStructSchemaId Id, const void* Struct)
+{
+	ReserveFlags(Id.Idx);
+	checkf(Flip(Bound[Id.Idx]), TEXT("'%s' already bound"), *GUE.Names.Print(Id));
+	check(!Instanced[Id.Idx]);
+	check(GetInstance(MakeStaticInstance(Struct)) == Struct);
+
+	Instanced[Id.Idx] = true;
+	Instances.Add(Id, MakeStaticInstance(Struct));
+}
+
 void FDefaultStructs::Drop(FStructSchemaId Id)
 {
+	checkf(!Flip(Bound[Id.Idx]), TEXT("'%s' isn't bound"), *GUE.Names.Print(Id));
 	if (Instanced[Id.Idx])
 	{
 		Instanced[Id.Idx] = false;
-		DeleteDefault(Instances.FindAndRemoveChecked(Id));
+		DeleteInstance(Instances.FindAndRemoveChecked(Id));
 	}
 }
 
 const void* FDefaultStructs::Get(FStructSchemaId Id)
 {
-	return Instanced[Id.Idx] ? Instances.FindChecked(Id) : Zeroes.GetData(); 
+	checkf(Bound[Id.Idx], TEXT("'%s' lack default"), *GUE.Names.Print(Id));
+	return Instanced[Id.Idx] ? GetInstance(Instances.FindChecked(Id)) : Zeroes.GetData(); 
 }
 
 ////////////////////////////////////////////////////////////////////////////////////////////////
 
 FCommonUObjectIds::FCommonUObjectIds(TIdIndexer<FName>& Names)
-: Core(				Names.MakeScope("/Script/Core"))
-, CoreUObject(		Names.MakeScope("/Script/CoreUObject"))
-, Name(				Names.IndexStruct({CoreUObject, Names.MakeTypename("Name")}))
-, Text(				Names.IndexStruct({CoreUObject, Names.MakeTypename("Text")}))
-, FieldPath(		Names.IndexStruct({CoreUObject, Names.MakeTypename("FieldPath")}))
-, SoftObjectPath(	Names.IndexStruct({CoreUObject, Names.MakeTypename("SoftObjectPath")}))
-, ObjectPtr(		Names.IndexStruct({CoreUObject, Names.MakeTypename("ObjectPtr")}))
-, WeakObjectPtr(	Names.IndexStruct({CoreUObject, Names.MakeTypename("WeakObjectPtr")}))
-, LazyObjectPtr(	Names.IndexStruct({CoreUObject, Names.MakeTypename("LazyObjectPtr")}))
-, SoftObjectPtr(	Names.IndexStruct({CoreUObject, Names.MakeTypename("SoftObjectPtr")}))
-, ScriptInterface(	Names.IndexStruct({CoreUObject, Names.MakeTypename("ScriptInterface")}))
-, Optional(			Names.NameType("Optional"))
-, Map(				Names.NameType("Map"))
-, Set(				Names.NameType("Set"))
-, Pair(				Names.NameType("Pair"))
-, Key(				Names.NameMember("Key"))
-, Value(			Names.NameMember("Value"))
-, Add(				Names.NameMember("Add"))
-, Del(				Names.NameMember("Del"))
+: Core(						Names.MakeScope("/Script/Core"))
+, CoreUObject(				Names.MakeScope("/Script/CoreUObject"))
+, Name(						Names.IndexStruct({Core, Names.MakeTypename("Name")}))
+, Text(						Names.IndexStruct({Core, Names.MakeTypename("Text")}))
+, Guid(						Names.IndexStruct({Core, Names.MakeTypename("Guid")}))
+, FieldPath(				Names.IndexStruct({CoreUObject, Names.MakeTypename("FieldPath")}))
+, SoftObjectPath(			Names.IndexStruct({CoreUObject, Names.MakeTypename("SoftObjectPath")}))
+, ObjectPtr(				Names.IndexStruct({CoreUObject, Names.MakeTypename("ObjectPtr")}))
+, WeakObjectPtr(			Names.IndexStruct({CoreUObject, Names.MakeTypename("WeakObjectPtr")}))
+, LazyObjectPtr(			Names.IndexStruct({CoreUObject, Names.MakeTypename("LazyObjectPtr")}))
+, SoftObjectPtr(			Names.IndexStruct({CoreUObject, Names.MakeTypename("SoftObjectPtr")}))
+, ScriptInterface(			Names.IndexStruct({CoreUObject, Names.MakeTypename("ScriptInterface")}))
+, Delegate(					Names.IndexStruct({CoreUObject, Names.MakeTypename("Delegate")}))
+, MulticastDelegate(		Names.IndexStruct({CoreUObject, Names.MakeTypename("MulticastDelegate")}))
+, MulticastInlineDelegate(	Names.IndexStruct({CoreUObject, Names.MakeTypename("MulticastInlineDelegate")}))
+, MulticastSparseDelegate(	Names.IndexStruct({CoreUObject, Names.MakeTypename("MulticastSparseDelegate")}))
+, Optional(					Names.NameType("Optional"))
+, Map(						Names.NameType("Map"))
+, Set(						Names.NameType("Set"))
+, Pair(						Names.NameType("Pair"))
+, Key(						Names.NameMember("Key"))
+, Value(					Names.NameMember("Value"))
+, Add(						Names.NameMember("Add"))
+, Del(						Names.NameMember("Del"))
+, Id(						Names.NameMember("Id"))
+, Object(					Names.NameMember("Object"))
+, Function(					Names.NameMember("Function"))
+, Path(						Names.NameMember("Path"))
+, Owner(					Names.NameMember("Owner"))
 {}
 
 ////////////////////////////////////////////////////////////////////////////////////////////////
@@ -147,7 +229,7 @@ FGlobals::FGlobals()
 : Types(Names)
 , Schemas(Names)
 , Customs(Names)
-, DeltaCustoms(Names, &Customs)
+, DeltaCustoms(Customs)
 , CommonIds(Names)
 {}
 
@@ -159,24 +241,26 @@ static constexpr uint64 LeafMask =			CASTCLASS_FNumericProperty | CASTCLASS_FEnu
 static constexpr uint64 IntSMask =			CASTCLASS_FInt8Property | CASTCLASS_FInt16Property | CASTCLASS_FIntProperty | CASTCLASS_FInt64Property;
 static constexpr uint64 IntUMask =			CASTCLASS_FByteProperty | CASTCLASS_FUInt16Property | CASTCLASS_FUInt32Property | CASTCLASS_FUInt64Property;
 static constexpr uint64 ContainerMask =		CASTCLASS_FArrayProperty | CASTCLASS_FSetProperty| CASTCLASS_FMapProperty  | CASTCLASS_FOptionalProperty;
-static constexpr uint64 StructuralMask =	CASTCLASS_FStructProperty | CASTCLASS_FDelegateProperty | CASTCLASS_FMulticastDelegateProperty;
 static constexpr uint64 StringMask =		CASTCLASS_FStrProperty | CASTCLASS_FUtf8StrProperty | CASTCLASS_FAnsiStrProperty;
-static constexpr uint64 CommonStructMask =	CASTCLASS_FNameProperty | CASTCLASS_FTextProperty | CASTCLASS_FInterfaceProperty | CASTCLASS_FFieldPathProperty |
-											CASTCLASS_FObjectProperty | CASTCLASS_FWeakObjectProperty | CASTCLASS_FSoftObjectProperty | CASTCLASS_FLazyObjectProperty;
+static constexpr uint64 CommonStructMask =	CASTCLASS_FNameProperty | CASTCLASS_FTextProperty |  CASTCLASS_FFieldPathProperty |
+											CASTCLASS_FObjectProperty | CASTCLASS_FWeakObjectProperty | CASTCLASS_FSoftObjectProperty | CASTCLASS_FLazyObjectProperty |
+											CASTCLASS_FDelegateProperty | CASTCLASS_FMulticastInlineDelegateProperty;
+static constexpr uint64 MiscMask =			CASTCLASS_FMulticastSparseDelegateProperty | CASTCLASS_FInterfaceProperty;
 
 static FStructSchemaId FlagsToCommonBindId(uint64 MaskedCastFlags)
 {
 	switch (MaskedCastFlags)
 	{
-		case CASTCLASS_FNameProperty:			return GUE.CommonIds.Name;
-		case CASTCLASS_FTextProperty:			return GUE.CommonIds.Text;
-		case CASTCLASS_FInterfaceProperty:		return GUE.CommonIds.ScriptInterface;
-		case CASTCLASS_FFieldPathProperty:		return GUE.CommonIds.FieldPath;
-		case CASTCLASS_FObjectProperty:			return GUE.CommonIds.ObjectPtr;
-		case CASTCLASS_FWeakObjectProperty:		return GUE.CommonIds.WeakObjectPtr;
-		case CASTCLASS_FSoftObjectProperty:		return GUE.CommonIds.SoftObjectPtr;
-		case CASTCLASS_FLazyObjectProperty:		return GUE.CommonIds.LazyObjectPtr;
-		default:								break; // error
+		case CASTCLASS_FNameProperty:						return GUE.CommonIds.Name;
+		case CASTCLASS_FObjectProperty:						return GUE.CommonIds.ObjectPtr;
+		case CASTCLASS_FWeakObjectProperty:					return GUE.CommonIds.WeakObjectPtr;
+		case CASTCLASS_FSoftObjectProperty:					return GUE.CommonIds.SoftObjectPtr;
+		case CASTCLASS_FLazyObjectProperty:					return GUE.CommonIds.LazyObjectPtr;
+		case CASTCLASS_FDelegateProperty:					return GUE.CommonIds.Delegate;
+		case CASTCLASS_FMulticastInlineDelegateProperty:	return GUE.CommonIds.MulticastInlineDelegate;
+		case CASTCLASS_FTextProperty:						return GUE.CommonIds.Text;
+		case CASTCLASS_FFieldPathProperty:					return GUE.CommonIds.FieldPath;
+		default:											break; // error
 	}
 
 	check(MaskedCastFlags); // @pre violated
@@ -194,8 +278,38 @@ static bool HasAny(uint64 Flags)
 
 ////////////////////////////////////////////////////////////////////////////////////////////////
 
+static bool ShouldBind(const UStruct* Struct);
+
+inline bool ShouldBind(const FProperty* Property)
+{
+	if (HasAny<CASTCLASS_FStructProperty>(Property->GetCastFlags()))
+	{
+		return ShouldBind(static_cast<const FStructProperty*>(Property)->Struct);
+	}
+
+	return true;
+}
+
+static bool ShouldBind(const UStruct* Struct)
+{
+	for (FProperty* Property = Struct->PropertyLink; Property; Property = Property->PropertyLinkNext)
+	{
+		if (ShouldBind(Property))
+		{
+			return true;
+		}
+	}
+	return false;
+}
+
+static const UStruct* SkipEmptyBases(const UStruct* In)
+{
+	return In->PropertyLink ? In->PropertyLink->GetOwnerChecked<UStruct>() : In;
+}
+
 static FTypeId IndexType(const UField* Field)
-{	
+{
+	check(Field);
 	FTypenameId Name = GUE.Names.MakeTypename(Field->GetFName());
 	
 	// This wouldn't be needed in an intrusive or cached solution
@@ -210,7 +324,16 @@ static FTypeId IndexType(const UField* Field)
 
 static FOptionalStructSchemaId IndexSuper(const UStruct* Struct)
 {
-	return Struct ? ToOptional(GUE.Names.IndexStruct(IndexType(Struct))) : NoId;
+	if (const UStruct* Super = Struct->GetInheritanceSuper())
+	{
+		if (ShouldBind(Super))
+		{
+			const UStruct* NonEmptySuper = SkipEmptyBases(Super);
+			return GUE.Names.IndexStruct(IndexType(NonEmptySuper));
+		}
+	}
+	
+	return NoId;
 }
 
 static EMemberPresence GetOccupancy(const UStruct* Struct)
@@ -227,12 +350,12 @@ static FStructSchemaId DeclareStruct(const UStruct* Struct)
 {
 	FTypeId Type = IndexType(Struct);
 	FStructSchemaId Id = GUE.Names.IndexStruct(Type);
-	FOptionalStructSchemaId  Super = IndexSuper(Struct->GetInheritanceSuper());
+	FOptionalStructSchemaId  Super = IndexSuper(Struct);
 	EMemberPresence Occupancy = GetOccupancy(Struct);
 	TArray<FMemberId, TInlineAllocator<64>> Members;
-	for (FProperty* Property = Struct->PropertyLink; Property; Property = Property->PropertyLinkNext)
+	for (FProperty* It = Struct->PropertyLink; It && It->GetOwner<UStruct>() == Struct; It = It->PropertyLinkNext)
 	{
-		Members.Add(GUE.Names.NameMember(Property->GetFName()));
+		Members.Add(GUE.Names.NameMember(It->GetFName()));
 	}
 	
 	GUE.Types.DeclareStruct(Id, Type, Members, Occupancy, Super);
@@ -242,6 +365,378 @@ static FStructSchemaId DeclareStruct(const UStruct* Struct)
 
 ////////////////////////////////////////////////////////////////////////////////////////////////
 
+static FTypedRange SaveNames(TConstArrayView<FName> Names, const FSaveContext& Ctx)
+{
+	FStructSchemaId Id = GUE.CommonIds.Name;
+	FStructRangeSaver Out(Ctx.Scratch, Names.Num());
+	for (FName Name : Names)
+	{
+		Out.AddItem(SaveStruct(&Name, Id, Ctx));
+	}
+	return Out.Finalize(MakeStructRangeSchema(DefaultRangeMax, Id));
+}
+
+static void LoadNames(TArray<FName>& Dst, FStructRangeLoadView Src)
+{
+	Dst.SetNumUninitialized(static_cast<int32>(Src.Num()));
+	FName* DstIt = Dst.GetData();
+	for (FStructLoadView Name : Src)
+	{
+		LoadStruct(DstIt++, Name);
+	}
+}
+
+void FFieldPathBinding::Save(FMemberBuilder& Dst, const FFieldPath& Src, const FFieldPath*, const FSaveContext& Ctx) const
+{
+	Dst.AddRange(GUE.CommonIds.Path, SaveNames(Src.*GFieldPathPath, Ctx));
+	Dst.AddStruct(GUE.CommonIds.Owner, GUE.CommonIds.WeakObjectPtr, SaveStruct(&(Src.*GFieldPathOwner), GUE.CommonIds.WeakObjectPtr, Ctx));
+}
+
+void FFieldPathBinding::Load(FFieldPath& Dst, FStructLoadView Src, ECustomLoadMethod Method) const
+{
+	FMemberLoader Members(Src);
+	Dst.Reset(); // ClearCachedField() more optimal
+	LoadStruct(&(Dst.*GFieldPathPath), Members.GrabStruct());
+	LoadStruct(&(Dst.*GFieldPathOwner), Members.GrabStruct());
+}
+
+bool FFieldPathBinding::Diff(const FFieldPath& A, const FFieldPath& B)
+{
+	return A != B;
+}
+
+////////////////////////////////////////////////////////////////////////////////////////////////
+
+void FDelegateBinding::Save(FMemberBuilder& Dst, const FScriptDelegate& Src, const FScriptDelegate* Default, const FSaveContext& Ctx) const
+{
+	if (Src.IsBound())
+	{
+		Dst.AddStruct(GUE.CommonIds.Object, GUE.CommonIds.WeakObjectPtr, SaveStruct(&(Src.*GDelegateObject), GUE.CommonIds.WeakObjectPtr, Ctx));
+		Dst.AddStruct(GUE.CommonIds.Function, GUE.CommonIds.Name, SaveStruct(&(Src.*GDelegateFunctionName), GUE.CommonIds.Name, Ctx));
+	}
+}
+
+void FDelegateBinding::Load(FScriptDelegate& Dst, FStructLoadView Src, ECustomLoadMethod Method) const
+{
+	FMemberLoader Members(Src);
+	if (Members.HasMore())
+	{
+		LoadStruct(&(Dst.*GDelegateObject), Members.GrabStruct());
+		LoadStruct(&(Dst.*GDelegateFunctionName), Members.GrabStruct());
+	}
+	else
+	{
+		Dst.Clear();
+	}
+}
+
+bool FDelegateBinding::Diff(const FScriptDelegate& A, const FScriptDelegate& B)
+{
+	return A != B;
+}
+
+////////////////////////////////////////////////////////////////////////////////////////////////
+
+static void SaveUnicastDelegate(FMemberBuilder& Dst, const FUnicastScriptDelegate& Src, const FSaveContext& Ctx)
+{
+	Dst.AddStruct(GUE.CommonIds.Object, GUE.CommonIds.WeakObjectPtr, SaveStruct(&(Src.*GUnicastDelegateObject), GUE.CommonIds.WeakObjectPtr, Ctx));
+	Dst.AddStruct(GUE.CommonIds.Function, GUE.CommonIds.Name, SaveStruct(&(Src.*GUnicastDelegateFunctionName), GUE.CommonIds.Name, Ctx));
+}
+
+static void LoadUnicastDelegate(FUnicastScriptDelegate& Dst, FStructLoadView Src)
+{
+	FMemberLoader Members(Src);
+	LoadStruct(&(Dst.*GUnicastDelegateObject), Members.GrabStruct());
+	LoadStruct(&(Dst.*GUnicastDelegateFunctionName), Members.GrabStruct());
+}
+
+static FTypedRange SaveInvocations(const FMulticastInvocationList& In, const FSaveContext& Ctx)
+{
+	FStructSchemaId ItemId = GUE.CommonIds.Delegate;
+	FMemberSchema Schema = MakeStructRangeSchema(DefaultRangeMax, ItemId);
+	if (int32 NumTotal = In.Num())
+	{
+		TBitArray<> Keep;
+		Keep.Reserve(NumTotal);
+		for (const FUnicastScriptDelegate& Invocation : In)
+		{
+			Keep.Add(!Invocation.IsCompactable());
+		}
+	
+		if (int32 NumKept = Keep.CountSetBits())
+		{
+			const FStructDeclaration& ItemDecl = GUE.Types.Get(ItemId);
+			const FUnicastScriptDelegate* Src = In.GetData();
+			FStructRangeSaver Dst(Ctx.Scratch, static_cast<uint64>(NumKept));
+			FMemberBuilder Tmp;
+			for (int32 Idx = 0; Idx < NumTotal; ++Idx)
+			{
+				if (Keep[Idx])
+				{
+					SaveUnicastDelegate(/* out */ Tmp, Src[Idx], Ctx);
+					Dst.AddItem(Tmp.BuildAndReset(Ctx.Scratch, ItemDecl, GUE.Names));
+				}
+			}
+			return Dst.Finalize(Schema);
+		}
+	}
+
+	return {Schema, nullptr};
+}
+
+static void SaveMulticastDelegate(FMemberBuilder& Dst, const FMulticastScriptDelegate& Src, const FSaveContext& Ctx)
+{
+	Dst.AddRange(GUE.CommonIds.Invocations, SaveInvocations(Src.*GMulticastDelegateInvocationList, Ctx));
+}
+
+static void SaveEmptyMulticastDelegate(FMemberBuilder& Dst)
+{
+	Dst.AddRange(GUE.CommonIds.Invocations, { MakeStructRangeSchema(DefaultRangeMax, GUE.CommonIds.Delegate), nullptr });
+}
+
+static void LoadInvocations(FMulticastInvocationList& Dst, FStructRangeLoadView Src)
+{
+	Dst.Reset(static_cast<int32>(Src.Num()));
+	for (FStructLoadView Invocation : Src)
+	{
+		LoadUnicastDelegate(Dst.AddDefaulted_GetRef(), Invocation);
+	}
+}
+
+static void LoadMulticastDelegate(FMulticastScriptDelegate& Dst, FMemberLoader& Src)
+{
+	LoadInvocations(Dst.*GMulticastDelegateInvocationList, Src.GrabRange().AsStructs());
+}
+
+inline bool DiffInvocations(TConstArrayView<FUnicastScriptDelegate> A, TConstArrayView<FUnicastScriptDelegate> B)
+{
+	if (A.Num() + B.Num() == 0)
+	{
+		return false;
+	}
+
+	const FUnicastScriptDelegate* EndA = A.GetData() + A.Num();
+	const FUnicastScriptDelegate* EndB = B.GetData() + B.Num();
+	for (const FUnicastScriptDelegate* ItA = A.GetData(), *ItB = B.GetData(); true; ++ItA, ++ItB)
+	{
+		for (; ItA != EndA && ItA->IsCompactable(); ++ItA) {}
+		for (; ItB != EndB && ItB->IsCompactable(); ++ItB) {}
+
+		if (ItA == EndA || ItB == EndB)
+		{
+			return ItA != EndA || ItB != EndB;
+		}
+		else if (*ItA != *ItB)
+		{
+			return true;
+		}
+	}
+}
+
+static bool DiffMulticastDelegate(const FMulticastScriptDelegate& A, const FMulticastScriptDelegate& B)
+{
+	return DiffInvocations(A.*GMulticastDelegateInvocationList, B.*GMulticastDelegateInvocationList);
+}
+
+void FMulticastInlineDelegateBinding::Save(FMemberBuilder& Dst, const FMulticastScriptDelegate& Src, const FMulticastScriptDelegate* Default, const FSaveContext& Ctx) const
+{
+	SaveMulticastDelegate(Dst, Src, Ctx);
+}
+
+void FMulticastInlineDelegateBinding::Load(FMulticastScriptDelegate& Dst, FStructLoadView Src, ECustomLoadMethod Method) const
+{
+	check(Method == ECustomLoadMethod::Assign);
+	FMemberLoader Members(Src);
+	LoadMulticastDelegate(Dst, Members);
+}
+
+bool FMulticastInlineDelegateBinding::Diff(const FMulticastScriptDelegate& A, const FMulticastScriptDelegate& B)
+{
+	return DiffMulticastDelegate(A, B);
+}
+
+////////////////////////////////////////////////////////////////////////////////////////////////
+
+struct FMulticastSparseDelegateBinding final : ICustomBinding
+{
+	explicit FMulticastSparseDelegateBinding(const USparseDelegateFunction* SignatureFunction)
+	: OwningClassName(SignatureFunction->OwningClassName)
+	, DelegateName(SignatureFunction->DelegateName)
+	{}
+
+	const FName OwningClassName;
+	const FName DelegateName;
+
+	virtual void SaveCustom(FMemberBuilder& Dst, const void* Src, const void* Default, const FSaveContext& Ctx) override
+	{
+		if (!Default || DiffCustom(Src, Default))
+		{
+			Save(Dst, *static_cast<const FSparseDelegate*>(Src), Ctx);
+		}
+	}
+
+	virtual void LoadCustom(void* Dst, FStructLoadView Src, ECustomLoadMethod Method) const override
+	{
+		check(Method == ECustomLoadMethod::Assign);
+		Load(*static_cast<FSparseDelegate*>(Dst), Src);
+	}
+
+	virtual bool DiffCustom(const void* A, const void* B) const override
+	{
+		return Diff(*static_cast<const FSparseDelegate*>(A), *static_cast<const FSparseDelegate*>(B));
+	}
+
+	const FMulticastScriptDelegate* GetMulticastDelegate(const FSparseDelegate& Sparse) const
+	{
+		if (Sparse.IsBound())
+		{
+			const UObject* Owner = FSparseDelegateStorage::ResolveSparseOwner(Sparse, OwningClassName, DelegateName);
+			return FSparseDelegateStorage::GetMulticastDelegate(Owner, DelegateName);
+		}
+		return nullptr;
+	}
+	
+	void Save(FMemberBuilder& Dst, const FSparseDelegate& Src, const FSaveContext& Ctx) const
+	{
+		if (const FMulticastScriptDelegate* Delegate = GetMulticastDelegate(Src))
+		{
+			SaveMulticastDelegate(Dst, *Delegate, Ctx);
+		}
+		else
+		{
+			SaveEmptyMulticastDelegate(Dst);
+		}
+	}
+	
+	void Load(FSparseDelegate& Dst, FStructLoadView Src) const
+	{
+		if (Dst.IsBound())
+		{
+			UObject* Owner = FSparseDelegateStorage::ResolveSparseOwner(Dst, OwningClassName, DelegateName);
+			FSparseDelegateStorage::Clear(Owner, DelegateName);
+			Dst.*GSparseDelegateIsBound = false;
+		}
+
+		FMemberLoader Members(Src);
+		if (Members.HasMore())
+		{
+			UObject* Owner = FSparseDelegateStorage::ResolveSparseOwner(Dst, OwningClassName, DelegateName);
+			FMulticastScriptDelegate Tmp;
+			LoadMulticastDelegate(Tmp, Members);
+			FSparseDelegateStorage::SetMulticastDelegate(Owner, DelegateName, MoveTemp(Tmp));
+			Dst.*GSparseDelegateIsBound = true;
+		}
+	}
+
+	bool Diff(const FSparseDelegate& SparseA, const FSparseDelegate& SparseB) const
+	{
+		const FMulticastScriptDelegate* A = GetMulticastDelegate(SparseA);
+		const FMulticastScriptDelegate* B = GetMulticastDelegate(SparseB);
+		if (A && B)
+		{
+			return DiffMulticastDelegate(*A, *B);
+		}
+		return !!A != !!B;
+	}
+};
+
+static FStructSchemaId BindSparseDelegate(FStructSchemaId Owner, FMulticastSparseDelegateProperty* Property)
+{
+	// Todo: Ownership / memory leak
+	ICustomBinding* Leak = new FMulticastSparseDelegateBinding(CastChecked<USparseDelegateFunction>(Property->SignatureFunction));
+
+	FTypeId MulticastSparseDelegate = GUE.Names.Resolve(GUE.CommonIds.MulticastSparseDelegate);
+	FTypeId OwnerParam = GUE.Names.Resolve(Owner);
+	FTypeId PropertyParam = { GUE.CommonIds.CoreUObject, FTypenameId(GUE.Names.NameType(Property->GetFName())) };
+	FTypeId UniqueBindName = GUE.Names.MakeParametricType(MulticastSparseDelegate, {OwnerParam, PropertyParam});
+	FStructSchemaId BindId = GUE.Names.IndexStruct(UniqueBindName);
+
+	GUE.Customs.BindStruct(BindId, *Leak, GUE.Types.Get(GUE.CommonIds.MulticastDelegate));
+
+	return BindId;
+}
+
+////////////////////////////////////////////////////////////////////////////////////////////////
+
+struct FInterfaceBinding final : ICustomBinding
+{
+	explicit FInterfaceBinding(UClass* Class) : InterfaceClass(Class) {}
+
+	const TObjectPtr<UClass> InterfaceClass;
+
+	virtual void SaveCustom(FMemberBuilder& Dst, const void* Src, const void* Default, const FSaveContext& Ctx) override
+	{
+		if (!Default || DiffCustom(Src, Default))
+		{
+			Save(Dst, *static_cast<const FScriptInterface*>(Src), Ctx);
+		}
+	}
+
+	virtual void LoadCustom(void* Dst, FStructLoadView Src, ECustomLoadMethod Method) const override
+	{
+		check(Method == ECustomLoadMethod::Assign);
+		Load(*static_cast<FScriptInterface*>(Dst), Src);
+	}
+
+	virtual bool DiffCustom(const void* A, const void* B) const override
+	{
+		return *static_cast<const FScriptInterface*>(A) != *static_cast<const FScriptInterface*>(B);
+	}
+	
+	void Save(FMemberBuilder& Dst, const FScriptInterface& Src, const FSaveContext& Ctx) const
+	{
+		Dst.AddStruct(GUE.CommonIds.Object, GUE.CommonIds.ObjectPtr, SaveStruct(Src.GetObject(), GUE.CommonIds.ObjectPtr, Ctx));
+	}
+
+	void Load(FScriptInterface& Dst, FStructLoadView Src) const
+	{
+		LoadStruct(Dst.GetObjectRef(), FMemberLoader(Src).GrabStruct());
+		UObject* Object = Dst.GetObject();
+		Dst.SetInterface(Object ? Object->GetInterfaceAddress(InterfaceClass) : nullptr);
+	}
+};
+
+class FInterfaceBindings
+{
+	const FTypeId ScriptInterface;
+	const FStructDeclaration& Declaration;
+
+	TMap<FTypeId, FStructSchemaId> BoundClasses;
+
+public:
+	FInterfaceBindings()
+	: ScriptInterface(GUE.Names.Resolve(GUE.CommonIds.ScriptInterface))
+	, Declaration(GUE.Types.DeclareStruct(GUE.CommonIds.ScriptInterface, ScriptInterface, {GUE.CommonIds.Object}, EMemberPresence::RequireAll))
+	{}
+
+	FStructSchemaId Bind(FInterfaceProperty* Property)
+	{
+		FTypeId Class = IndexType(Property->InterfaceClass);
+		if (const FStructSchemaId* BindId = BoundClasses.Find(Class))
+		{
+			return *BindId;
+		}
+		
+		FTypeId UniqueBindName = GUE.Names.MakeParametricType(ScriptInterface, {Class});
+		FStructSchemaId BindId = GUE.Names.IndexStruct(UniqueBindName);
+		BoundClasses.Emplace(Class, BindId);
+
+		// Todo: Ownership / memory leak
+		ICustomBinding* Leak = new FInterfaceBinding(Property->InterfaceClass);
+		GUE.Customs.BindStruct(BindId, *Leak, Declaration);
+
+		return BindId;
+	}
+};
+
+static FStructSchemaId BindInterface(FInterfaceProperty* Property)
+{
+	static FInterfaceBindings Bindings;
+	return Bindings.Bind(Property);
+}
+
+////////////////////////////////////////////////////////////////////////////////////////////////
+
 inline uint32 HashRangeBindings(TConstArrayView<FRangeBinding> In)
 {
 	return static_cast<uint32>(FXxHash64::HashBuffer(In.GetData(), In.NumBytes()).Hash);
@@ -334,13 +829,6 @@ inline bool HasHash(FProperty* Property)
 	return !!(Property->PropertyFlags & CPF_HasGetValueTypeHash);
 }
 
-// Temp hack. Long-term, try throwing away FProperty and implement Construct/DestroyValue replacement
-} // namespace PlainProps::UE
-UE_DEFINE_PRIVATE_MEMBER_PTR(void(void*) const, GInitPropertyValue, FProperty, InitializeValueInternal);
-UE_DEFINE_PRIVATE_MEMBER_PTR(void(void*) const, GDestroyPropertyValue, FProperty, DestroyValueInternal);
-UE_DEFINE_PRIVATE_MEMBER_PTR(uint32(const void*) const, GHashPropertyValue, FProperty, GetValueTypeHashInternal);
-namespace PlainProps::UE {
-
 inline void ConstructValue(const FProperty* Property, void* Value)
 {
 	((*Property).*GInitPropertyValue)(Value);
@@ -382,8 +870,6 @@ inline void DestroyValues(const FProperty* Property, uint8* Values, int32 Num, S
 
 ////////////////////////////////////////////////////////////////////////////////////////////////
 
-static constexpr ERangeSizeType DefaultSizeType = RangeSizeOf(FDefaultAllocator::SizeType{});
-
 // Helps cache array property range bindings
 struct FArrayPropertyInfo
 {
@@ -427,7 +913,7 @@ struct TTrivialArrayBinding : IItemRangeBinding
 	virtual void ReadItems(FSaveRangeContext& Ctx) const override
 	{
 		const ScriptArray& Array = Ctx.Request.GetRange<ScriptArray>();
-		Ctx.Items.SetAll({Array.GetData(), static_cast<uint64>(Array.Num())}, Info.ItemSize);
+		Ctx.Items.SetAll(Array.GetData(), static_cast<uint64>(Array.Num()), Info.ItemSize);
 	}
 
 	virtual void MakeItems(FLoadRangeContext& Ctx) const override
@@ -540,7 +1026,7 @@ struct TLeafArrayBinding : ILeafRangeBinding
 // Reusable cache of FArrayProperty range bindings
 class FArrayPropertyBindings
 {
-	static constexpr ERangeSizeType SizeType = DefaultSizeType;
+	static constexpr ERangeSizeType SizeType = DefaultRangeMax;
 	
 	const TLeafArrayBinding<ELeafType::Bool, ELeafWidth::B8>	Bool;
 	const TLeafArrayBinding<ELeafType::Float, ELeafWidth::B32>	Float;
@@ -557,6 +1043,8 @@ class FArrayPropertyBindings
 
 	TMap<FArrayPropertyInfo, IItemRangeBinding*>				Others;
 
+	inline const ILeafRangeBinding& DownCast(const ILeafRangeBinding& In) { return In; }
+
 public:
 	FArrayPropertyBindings()
 	: Integers{{FRangeBinding(IntU8, SizeType),	FRangeBinding(IntU16, SizeType), FRangeBinding(IntU32, SizeType), FRangeBinding(IntU64, SizeType)},
@@ -586,7 +1074,7 @@ public:
 			
 			check(HasAny<CASTCLASS_FFloatProperty | CASTCLASS_FDoubleProperty | CASTCLASS_FBoolProperty >(InnerCastFlags));
 			const ILeafRangeBinding& Binding = HasAny<CASTCLASS_FBoolProperty>(InnerCastFlags) 
-				? (HasAny<CASTCLASS_FFloatProperty>(InnerCastFlags) ? static_cast<const ILeafRangeBinding&>(Float) : Double) : Bool;
+				? DownCast(Bool) : (HasAny<CASTCLASS_FFloatProperty>(InnerCastFlags) ? DownCast(Float) : Double);
 			return FRangeBinding(Binding, SizeType);
 		}
 		else if (IItemRangeBinding** Cached = Others.Find(Info))
@@ -887,7 +1375,7 @@ struct TLeafRangeSerializer
 
 	inline FMemberSchema MakeMemberSchema() const
 	{
-		return { FMemberType(DefaultSizeType), InnerType, 1, EnumId, nullptr };
+		return { FMemberType(DefaultRangeMax), InnerType, 1, EnumId, nullptr };
 	}
 
 	inline static EquivalentType SaveItem(const void* In, const FSaveContext&)
@@ -922,7 +1410,7 @@ struct FStructRangeSerializer
 
 	inline FMemberSchema MakeMemberSchema() const
 	{
-		return { FMemberType(DefaultSizeType), InnerType, 1, FOptionalSchemaId(Id), nullptr };
+		return { FMemberType(DefaultRangeMax), InnerType, 1, FOptionalSchemaId(Id), nullptr };
 	}
 
 	inline FBuiltStruct* SaveItem(const void* In, const FSaveContext& Ctx) const
@@ -964,7 +1452,7 @@ struct FNestedRangeSerializer
 
 	inline FMemberSchema MakeMemberSchema() const
 	{
-		return { FMemberType(DefaultSizeType), InnerTypes[0], NumInners, InnermostId, InnerTypes.GetData() };
+		return { FMemberType(DefaultRangeMax), InnerTypes[0], NumInners, InnermostId, InnerTypes.GetData() };
 	}
 
 	FBuiltRange* SaveItem(const void* In, const FSaveContext& Ctx) const
@@ -976,7 +1464,7 @@ struct FNestedRangeSerializer
 	inline void LoadItem(void* Dst, FByteReader& SrcBytes, FBitCacheReader& SrcBits, const FLoadBatch& Batch) const
 	{
 		FRangeLoadSchema Schema = {InnerTypes[0], InnermostId, InnerTypes.GetData() + 1, Batch};
-		LoadRange(Dst, SrcBytes, SrcBits, DefaultSizeType, Schema, InnerBindings);
+		LoadRange(Dst, SrcBytes, SrcBits, DefaultRangeMax, Schema, InnerBindings);
 	}
 };
 
@@ -1039,8 +1527,8 @@ inline void ReadSparseItems(FExistingItems& Dst, const FScriptSparseArray& Src,
 	}
 	else if (Src.IsCompact())
 	{
-		int32 Num = Src.Num();
-		Dst.SetAll(Num ? Data : nullptr, Num);
+		uint64 Num = static_cast<uint64>(Src.Num());
+		Dst.SetAll(Num ? Data : nullptr, Num, Layout.Size);
 	}
 	else
 	{
@@ -1455,27 +1943,27 @@ class FSetBindings
 	TMap<FParameterBinding, FRangeBinding> RangeBindings;
 
 	template<EPropertyKind ItemKind>
-	IItemRangeBinding* BindNew(FSetProperty* Property, FMemberBinding Elem, FStructSchemaId Id)
+	IItemRangeBinding* BindNew(FSetProperty* Property, FMemberBinding Elem, const FStructDeclaration& Declaration) 
 	{
 		// Todo: Ownership / memory leak
 		TSetPropertyBinding<ItemKind>* Out = new TSetPropertyBinding<ItemKind>(Property, Elem);
-		GUE.DeltaCustoms.BindStruct(Id, Id, {}, *Out);
+		GUE.DeltaCustoms.BindStruct(Declaration.Id, *Out, Declaration, {});
 		return Out;
 	}
 
-	IItemRangeBinding* DeltaBindNew(FSetProperty* Property, FMemberBinding Elem, FStructSchemaId Id)
+	IItemRangeBinding* DeltaBindNew(FSetProperty* Property, FMemberBinding Elem, const FStructDeclaration& Declaration)
 	{
 		switch (GetPropertyKind(Elem))
 		{
-		case EPropertyKind::Range:		return BindNew<EPropertyKind::Range>(Property, Elem, Id);
-		case EPropertyKind::Struct:		return BindNew<EPropertyKind::Struct>(Property, Elem, Id);
-		case EPropertyKind::Bool:		return BindNew<EPropertyKind::Bool>(Property, Elem, Id);
-		case EPropertyKind::U8:			return BindNew<EPropertyKind::U8>(Property, Elem, Id);
-		case EPropertyKind::U16:		return BindNew<EPropertyKind::U16>(Property, Elem, Id);
-		case EPropertyKind::U32:		return BindNew<EPropertyKind::U32>(Property, Elem, Id);
-		case EPropertyKind::U64:		return BindNew<EPropertyKind::U64>(Property, Elem, Id);
-		case EPropertyKind::F32:		return BindNew<EPropertyKind::F32>(Property, Elem, Id);
-		case EPropertyKind::F64:		return BindNew<EPropertyKind::F64>(Property, Elem, Id);
+		case EPropertyKind::Range:		return BindNew<EPropertyKind::Range>(Property, Elem, Declaration);
+		case EPropertyKind::Struct:		return BindNew<EPropertyKind::Struct>(Property, Elem, Declaration);
+		case EPropertyKind::Bool:		return BindNew<EPropertyKind::Bool>(Property, Elem, Declaration);
+		case EPropertyKind::U8:			return BindNew<EPropertyKind::U8>(Property, Elem, Declaration);
+		case EPropertyKind::U16:		return BindNew<EPropertyKind::U16>(Property, Elem, Declaration);
+		case EPropertyKind::U32:		return BindNew<EPropertyKind::U32>(Property, Elem, Declaration);
+		case EPropertyKind::U64:		return BindNew<EPropertyKind::U64>(Property, Elem, Declaration);
+		case EPropertyKind::F32:		return BindNew<EPropertyKind::F32>(Property, Elem, Declaration);
+		case EPropertyKind::F64:		return BindNew<EPropertyKind::F64>(Property, Elem, Declaration);
 		default: check(false);			return nullptr;
 		}
 	}
@@ -1496,9 +1984,9 @@ public:
 		FMemberId Members[2] = { GUE.CommonIds.Del, GUE.CommonIds.Add };
 
 		// Todo: Ownership / memory leak
-		GUE.Types.DeclareStruct(Id, Type, Members, EMemberPresence::AllowSparse);
-		IItemRangeBinding* Out = DeltaBindNew(Property, Elem, Id);
-		return RangeBindings.Emplace(Elem, FRangeBinding(*Out, DefaultSizeType));
+		const FStructDeclaration& Declaration = GUE.Types.DeclareStruct(Id, Type, Members, EMemberPresence::AllowSparse);
+		IItemRangeBinding* Out = DeltaBindNew(Property, Elem, Declaration);
+		return RangeBindings.Emplace(Elem, FRangeBinding(*Out, DefaultRangeMax));
 	}
 };
 static FSetBindings GSets;
@@ -1818,6 +2306,7 @@ struct TMapPropertyDeltaBinding : TMapPropertyBinding<FScriptMap, KeyKind, Value
 
 ////////////////////////////////////////////////////////////////////////////////////////////////
 
+
 class FMapBindings
 {
 	TMap<FStructSchemaId, FRangeBinding> NormalBindings;
@@ -1830,7 +2319,19 @@ class FMapBindings
 		{
 			return new TMapPropertyBinding<FFreezableScriptMap, KeyKind, ValueKind>(Property);
 		}
-		return new TMapPropertyDeltaBinding<KeyKind, ValueKind>(Property, Members);
+
+		// Index custom delta binding struct name
+		FParametricTypeId PairTypename = GUE.Names.Resolve(static_cast<FStructSchemaId>(Members.Pair.InnermostSchema.Get())).Name.AsParametric();
+		TConstArrayView<FTypeId> Params = GUE.Names.Resolve(PairTypename).GetParameters();
+		FTypeId Type = { GUE.CommonIds.Core, FTypenameId(GUE.Names.MakeParametricTypeId(GUE.CommonIds.Map, Params)) };
+		FStructSchemaId Id = GUE.Names.IndexStruct(Type);
+
+		// Todo: Ownership / memory leak
+		const FStructDeclaration& Declaration = GUE.Types.DeclareStruct(Id, Type, { GUE.CommonIds.Del, GUE.CommonIds.Add }, EMemberPresence::AllowSparse);
+		auto Out = new TMapPropertyDeltaBinding<KeyKind, ValueKind>(Property, Members);
+		GUE.DeltaCustoms.BindStruct(Declaration.Id, *Out, Declaration, {});
+
+		return Out;
 	}
 
 	template<EPropertyKind KeyKind>
@@ -1881,7 +2382,7 @@ public:
 
 		// Todo: Ownership / memory leak
 		IItemRangeBinding* Leak = New(Property, Members, bFreezable);
-		return Bindings.Emplace(PairId, FRangeBinding(*Leak, DefaultSizeType));
+		return Bindings.Emplace(PairId, FRangeBinding(*Leak, DefaultRangeMax));
 	}
 };
 static FMapBindings GMaps;
@@ -1976,14 +2477,15 @@ public:
 static FOptionalBindings GOptionals;
 
 //////////////////////////////////////////////////////////////////////////////////////////////
+
 struct FStringBindings
 {
 	TStringBinding<FString>		TChrInstance;
 	TStringBinding<FAnsiString>	AnsiInstance;
 	TStringBinding<FUtf8String>	Utf8Instance;
-	FRangeBinding				TChr{TChrInstance, DefaultSizeType};
-	FRangeBinding				Ansi{AnsiInstance, DefaultSizeType};
-	FRangeBinding				Utf8{Utf8Instance, DefaultSizeType};
+	FRangeBinding				TChr{TChrInstance, DefaultRangeMax};
+	FRangeBinding				Ansi{AnsiInstance, DefaultRangeMax};
+	FRangeBinding				Utf8{Utf8Instance, DefaultRangeMax};
 
 	FMemberBinding				Bind(FProperty* Property, uint64 CastFlags) const
 	{
@@ -2011,7 +2513,7 @@ struct FStaticArrayBinding : IItemRangeBinding
 
 	virtual void ReadItems(FSaveRangeContext& Ctx) const override
 	{
-		Ctx.Items.SetAll({Ctx.Request.Range, Num}, Stride);
+		Ctx.Items.SetAll(Ctx.Request.Range, Num, Stride);
 	}
 
 	virtual void MakeItems(FLoadRangeContext& Ctx) const override
@@ -2106,8 +2608,7 @@ private:
 
 	inline FMemberBinding BindAsStruct(FProperty* Property, UStruct* Struct)
 	{
-		check(Struct);
-		FTypeId Type = IndexType(Struct);
+		FTypeId Type = IndexType(SkipEmptyBases(Struct));
 		return BindAsStruct(Property, GUE.Names.IndexStruct(Type));
 	}
 
@@ -2209,15 +2710,9 @@ private:
 		{
 			return BindAsStruct(Property, FlagsToCommonBindId(Flags & CommonStructMask));
 		}
-		else if (HasAny<StructuralMask>(Flags))
+		else if (HasAny<CASTCLASS_FStructProperty>(Flags))
 		{
-			if (HasAny<CASTCLASS_FStructProperty>(Flags))
-			{
-				return BindAsStruct(Property, static_cast<FStructProperty*>(Property)->Struct);
-			}
-			return BindAsStruct(Property, HasAny<CASTCLASS_FDelegateProperty>(Flags) 
-								? static_cast<FDelegateProperty*>(Property)->SignatureFunction
-								: static_cast<FMulticastDelegateProperty*>(Property)->SignatureFunction);
+			return BindAsStruct(Property, static_cast<FStructProperty*>(Property)->Struct);
 		}
 		else if (HasAny<ContainerMask>(Flags))
 		{
@@ -2237,10 +2732,17 @@ private:
 		{
 			return GStrings.Bind(Property, Flags);
 		}
+		else if (HasAny<MiscMask>(Flags))
+		{
+			FStructSchemaId BindId = HasAny<CASTCLASS_FInterfaceProperty>(Flags)
+				? BindInterface(static_cast<FInterfaceProperty*>(Property))
+				: BindSparseDelegate(Owner, static_cast<FMulticastSparseDelegateProperty*>(Property));
+			return BindAsStruct(Property, BindId);
+		}
 	#if WITH_VERSE_VM
 		else if (HasAny<CASTCLASS_FVValueProperty | CASTCLASS_FVRestValueProperty>(Flags))
 		{
-			return HasAny<CASTCLASS_FEnumProperty>(Flags)
+			return HasAny<CASTCLASS_FVValueProperty>(Flags)
 				? BindVValue(static_cast<FVValueProperty*>(Property))
 				: BindVRestValue(static_cast<FVRestValueProperty*>(Property)); 
 		}
@@ -2320,9 +2822,12 @@ void BindStruct(FStructSchemaId Id, UStruct* Struct)
 	}
 
 	FPropertyBinder Binder(Id);
-	for (FProperty* Property = Struct->PropertyLink; Property; Property = Property->PropertyLinkNext)
+	for (FProperty* It = Struct->PropertyLink; It && It->GetOwner<UStruct>() == Struct; It = It->PropertyLinkNext)
 	{
-		Binder.BindMember(Property);
+		if (ShouldBind(It))
+		{
+			Binder.BindMember(It);
+		}
 	}
 	GUE.Schemas.BindStruct(Id, Id, Binder.GetMembers());
 	
@@ -2337,27 +2842,24 @@ void BindStruct(FStructSchemaId Id, UStruct* Struct)
 
 static void BindInitialTypes()
 {
-	if (!FParse::Param(FCommandLine::Get(), TEXT("BindPlainProps")))
-	{
-		return;
-	}
-
 	UE_LOG(LogInit, Display, TEXT("Binding types to PlainProps schemas"));
 
+	// Declare all UScriptStruct/UClass/UFunction/UEnum
 	TArray<FStructSchemaId> Ids;
 	for (TObjectIterator<UField> It; It; ++It)
 	{
-		UObject* Object = *It; 
-		if (UStruct* Struct = Cast<UStruct>(Object))
+		UField* Field = *It;
+		if (UStruct* Struct = Cast<UStruct>(Field))
 		{
 			Ids.Add(DeclareStruct(Struct));
 		}
-		else if (UEnum* Enum = Cast<UEnum>(Object))
+		else if (UEnum* Enum = Cast<UEnum>(Field))
 		{
 			DeclareEnum(Enum);
-		}	
+		}
 	}
 
+	// Bind all UScriptStruct/UClass/UFunction
 	const FStructSchemaId* IdIt = &Ids[0];
 	for (TObjectIterator<UStruct> It; It; ++It)
 	{
@@ -2366,6 +2868,334 @@ static void BindInitialTypes()
 	check(IdIt == &Ids[0] + Ids.Num());
 }
 
-static FDelegateHandle PreMainCallback = FCoreDelegates::OnCommandletPreMain.AddStatic(&BindInitialTypes);
+static void DeclareProperty(FStructSchemaId Id, TConstArrayView<FMemberId> Members, EMemberPresence Occupancy)
+{
+	GUE.Types.DeclareStruct(Id,	GUE.Names.Resolve(Id), Members, Occupancy);
+}
+
+static void InitBatchedProperties()
+{
+	// All property batch types share declaration, but some members could be exclusive to a specific type e.g. FMemoryPropertyBatch
+	DeclareProperty(GUE.CommonIds.Name,				{ GUE.CommonIds.Id }, EMemberPresence::RequireAll);
+	DeclareProperty(GUE.CommonIds.Text,				{ GUE.CommonIds.Id }, EMemberPresence::RequireAll);
+	DeclareProperty(GUE.CommonIds.ObjectPtr,		{ GUE.CommonIds.Id }, EMemberPresence::RequireAll);
+	DeclareProperty(GUE.CommonIds.WeakObjectPtr,	{ GUE.CommonIds.Id }, EMemberPresence::RequireAll);
+	DeclareProperty(GUE.CommonIds.SoftObjectPtr,	{ GUE.CommonIds.Id }, EMemberPresence::RequireAll);
+	DeclareProperty(GUE.CommonIds.LazyObjectPtr,	{ GUE.CommonIds.Id }, EMemberPresence::RequireAll);
+
+	GUE.Defaults.BindZeroes(GUE.CommonIds.Name, sizeof(FName), alignof(FName));
+	GUE.Defaults.BindStatic(GUE.CommonIds.Text, &FText::GetEmpty());
+	GUE.Defaults.BindZeroes(GUE.CommonIds.ObjectPtr, sizeof(FObjectPtr), alignof(FObjectPtr));
+	GUE.Defaults.BindZeroes(GUE.CommonIds.WeakObjectPtr, sizeof(FWeakObjectPtr), alignof(FWeakObjectPtr));
+	GUE.Defaults.BindZeroes(GUE.CommonIds.SoftObjectPtr, sizeof(FSoftObjectPtr), alignof(FSoftObjectPtr));
+	GUE.Defaults.BindZeroes(GUE.CommonIds.LazyObjectPtr, sizeof(FLazyObjectPtr), alignof(FLazyObjectPtr));
+}
+
+struct FMemoryPropertyBatch
+{
+	TArray<FText> Texts;
+
+	static void Save(FMemberBuilder& Out, FName In, const FSaveContext&)
+	{
+		Out.Add(GUE.CommonIds.Id, In.ToUnstableInt());
+	}
+
+	static void Load(FName& Out, FStructLoadView In)
+	{
+		// Todo: embed in some experimental load as leaf API
+		In.Values.CheckSize(sizeof(decltype(ToInt(FName()))));
+		Out = FromInt(FPlatformMemory::ReadUnaligned<uint64>(In.Values.Peek()));
+	}
+
+	void Save(FMemberBuilder& Out, const FText& In, const FSaveContext&)
+	{
+		Out.Add(GUE.CommonIds.Id, Texts.Num());
+		Texts.Add(In);
+	}
+
+	void Load(FText& Out, FStructLoadView In) const
+	{
+		// Todo: embed in some experimental load as leaf API
+		In.Values.CheckSize(sizeof(int32));
+		int32 Idx = FPlatformMemory::ReadUnaligned<int32>(In.Values.Peek());
+		Out = Texts[Idx];
+	}
+
+	static void Save(FMemberBuilder& Out, FObjectHandle In, const FSaveContext&)
+	{
+		static_assert(sizeof(In) == sizeof(uint64));
+		Out.Add(GUE.CommonIds.Id, reinterpret_cast<uint64&>(In));
+	}
+
+	static void Load(FObjectHandle& Out, FStructLoadView In)
+	{
+		// Todo: embed in some experimental load as leaf API
+		In.Values.CheckSize(sizeof(uint64));
+		reinterpret_cast<uint64&>(Out) = FPlatformMemory::ReadUnaligned<uint64>(In.Values.Peek());
+	}
+
+	static void Save(FMemberBuilder& Out, const FWeakObjectPtr& In, const FSaveContext&)
+	{
+		// Save ObjectSerialNumber + ObjectIndex a single uint64
+		static_assert(sizeof(FWeakObjectPtr) == sizeof(uint64));
+		Out.Add(GUE.CommonIds.Id, reinterpret_cast<const uint64&>(In));
+	}
+
+	static void Load(FWeakObjectPtr& Out, FStructLoadView In)
+	{
+		// Todo: embed in some experimental load as leaf API
+		In.Values.CheckSize(sizeof(uint64));
+		reinterpret_cast<uint64&>(Out) = FPlatformMemory::ReadUnaligned<uint64>(In.Values.Peek());
+	}
+
+	static void Save(FMemberBuilder& Out, const FSoftObjectPtr& In, const FSaveContext& Ctx)
+	{
+		FBuiltStruct* SoftPath = SaveStruct(&In.GetUniqueID(), GUE.CommonIds.SoftObjectPath, Ctx);
+		Out.AddStruct(GUE.CommonIds.Id, GUE.CommonIds.SoftObjectPath, SoftPath);
+	}
+
+	static void Load(FSoftObjectPtr& Out, FStructLoadView In)
+	{
+		Out.ResetWeakPtr();
+		FMemberLoader Members(In);
+		LoadStruct(&Out.GetUniqueID(), Members.GrabStruct());
+	}
+
+	static void Save(FMemberBuilder& Out, const FLazyObjectPtr& In, const FSaveContext& Ctx)
+	{
+		FBuiltStruct* Guid = SaveStruct(&In.GetUniqueID(), GUE.CommonIds.Guid, Ctx);
+		Out.AddStruct(GUE.CommonIds.Id, GUE.CommonIds.Guid, Guid);
+	}
+
+	static void Load(FLazyObjectPtr& Out, FStructLoadView In)
+	{
+		Out.ResetWeakPtr();
+		FMemberLoader Members(In);
+		LoadStruct(&Out.GetUniqueID(), Members.GrabStruct());
+	}
+
+#if UE_FNAME_OUTLINE_NUMBER
+	static uint32 ToInt(FName In)
+	{
+		return In.GetDisplayIndex().ToUnstableInt();
+	}
+	static FName FromInt(uint32 In)
+	{
+		return FName::CreateFromDisplayId(FNameEntryId::FromUnstableInt(In), NAME_NO_NUMBER_INTERNAL);
+	}
+#else
+	static uint64 ToInt(FName In)
+	{
+		return (uint64(In.GetNumber()) << 32) | In.GetDisplayIndex().ToUnstableInt();
+	}
+	static FName FromInt(uint64 In)
+	{
+		return FName::CreateFromDisplayId(FNameEntryId::FromUnstableInt(static_cast<uint32>(In)), static_cast<uint32>(In >> 32));
+	}
+#endif
+};
+
+inline bool DiffProperty(FName A, FName B) { return !A.IsEqual(B, ENameCase::CaseSensitive); }
+inline bool DiffProperty(const FText& A, const FText& B) { return !FTextProperty::Identical_Implementation(A, B, 0); }
+inline bool DiffProperty(FObjectHandle A, FObjectHandle B) { return A != B; }
+inline bool DiffProperty(const FWeakObjectPtr& A, const FWeakObjectPtr& B) { return A != B; }
+inline bool DiffProperty(const FSoftObjectPtr& A, const FSoftObjectPtr& B) { return A != B; }
+inline bool DiffProperty(const FLazyObjectPtr& A, const FLazyObjectPtr& B) { return A != B; }
+
+
+template<class Type, class BatchType>
+struct TCustomPropertyBinding final : ICustomBinding
+{
+	TCustomPropertyBinding(BatchType& InBatch) : Batch(InBatch) {}
+
+	BatchType& Batch;
+
+	virtual void SaveCustom(FMemberBuilder& Dst, const void* Src, const void* Default, const FSaveContext& Ctx) override
+	{
+		if (!Default || DiffCustom(Src, Default))
+		{
+			Batch.Save(Dst, *static_cast<const Type*>(Src), Ctx);
+		}
+	}
+
+	virtual void LoadCustom(void* Dst, FStructLoadView Src, ECustomLoadMethod) const override
+	{
+		Batch.Load(*static_cast<Type*>(Dst), Src);
+	}
+
+	virtual bool DiffCustom(const void* A, const void* B) const override
+	{
+		return DiffProperty(*static_cast<const Type*>(A), *static_cast<const Type*>(B));
+	}
+};
+
+template<class BatchType>
+struct TCustomPropertyBindings
+{
+	TCustomPropertyBindings(BatchType& Batch, const FCustomBindings& Underlay)
+	: Overlay(Underlay)
+	, Name(Batch)
+	, Text(Batch)
+	, ObjectPtr(Batch)
+	, SoftObjectPtr(Batch)
+	, WeakObjectPtr(Batch)
+	, LazyObjectPtr(Batch)
+	{
+		Bind(GUE.CommonIds.Name, Name);
+		Bind(GUE.CommonIds.Text, Text);
+		Bind(GUE.CommonIds.ObjectPtr, ObjectPtr);
+		Bind(GUE.CommonIds.SoftObjectPtr, SoftObjectPtr);
+		Bind(GUE.CommonIds.WeakObjectPtr, WeakObjectPtr);
+		Bind(GUE.CommonIds.LazyObjectPtr, LazyObjectPtr);
+	}
+
+	void Bind(FStructSchemaId Id, ICustomBinding& Binding)
+	{
+		Overlay.BindStruct(Id, Binding, GUE.Types.Get(Id));
+	}
+
+	FCustomBindingsOverlay								Overlay;
+	TCustomPropertyBinding<FName, BatchType>			Name;
+	TCustomPropertyBinding<FText, BatchType>			Text;
+	TCustomPropertyBinding<FObjectHandle, BatchType>	ObjectPtr;
+	TCustomPropertyBinding<FSoftObjectPtr, BatchType>	SoftObjectPtr;
+	TCustomPropertyBinding<FWeakObjectPtr, BatchType>	WeakObjectPtr;
+	TCustomPropertyBinding<FLazyObjectPtr, BatchType>	LazyObjectPtr;
+};
+
+template<class PropertyBatchType>
+class TBatchSaver
+{
+public:
+	TBatchSaver(int32 NumReserve)
+	: FlatBindings(SavedProperties, GUE.Customs)
+	, DeltaBindings(SavedProperties, GUE.DeltaCustoms)
+	, FlatCtx{GUE.Types, GUE.Schemas, FlatBindings.Overlay, Scratch, nullptr}
+	, DeltaCtx{GUE.Types, GUE.Schemas, DeltaBindings.Overlay, Scratch, &GUE.Defaults}
+	{
+		SavedObjects.Reserve(NumReserve);
+	}
+
+	void Save(const UObject* Object)
+	{
+		const UClass* Class = Object->GetClass();
+		const UObject* CDO = Class->GetDefaultObject(/* create lazily */ false);
+		FStructSchemaId Id = GUE.Names.IndexStruct(IndexType(Class));
+		bool bFlat = Object == CDO || CDO == nullptr;
+		SavedObjects.Emplace(Id, bFlat ? SaveStruct(Object, Id, FlatCtx) : SaveStructDelta(Object, CDO, Id, DeltaCtx));
+	}
+
+	//TArray64<uint8>	Write() const
+	//{
+	//	// Build partial schemas
+	//	const FStructBindIds BindIds(GUE.DeltaCustoms, GUE.Schemas);
+	//	FSchemasBuilder SchemaBuilders(GUE.Types, BindIds, Scratch);
+	//	for (const IdBuiltStructPair& Object : SavedObjects)
+	//	{
+	//		SchemaBuilders.NoteStructAndMembers(Object.Key, *Object.Value);
+	//	}
+	//	FBuiltSchemas Schemas = SchemaBuilders.Build(); 
+
+	//	// Filter out declared but unused names and ids
+	//	FWriter Writer(GUE.Names, BindIds, Schemas, ESchemaFormat::InMemoryNames);
+	//	TArray<FName> UsedNames;
+	//	UsedNames.Reserve(Writer.GetUsedNames().Num());
+	//	for (FNameId Name : Writer.GetUsedNames())
+	//	{
+	//		UsedNames.Add(GUE.Names.ResolveName(Name));
+	//	}
+	//
+	//	// Write ids. Just copying in-memory FNames, a stable format might use SaveNameBatch().
+	//	TArray64<uint8> Out;
+	//	WriteU32(Out, Magics[0]);
+	//	WriteNumAndArray(Out, TArrayView<const FName, int32>(UsedNames));
+
+	//	// Write schemas
+	//	WriteU32(Out, Magics[1]);
+	//	WriteAlignmentPadding<uint32>(Out);
+	//	TArray64<uint8> Tmp;
+	//	Writer.WriteSchemas(/* Out */ Tmp);
+	//	WriteNumAndArray(Out, TArrayView<const uint8, int64>(Tmp));
+	//	Tmp.Reset();
+
+	//	// Write objects
+	//	WriteU32(Out, Magics[2]);
+	//	for (const TPair<FStructSchemaId, FBuiltStruct*>& Object : SavedObjects)
+	//	{
+	//		WriteU32(/* out */ Tmp, Magics[3]);
+	//		WriteU32(/* out */ Tmp, Writer.GetWriteId(Object.Key).Get().Idx);
+	//		Writer.WriteMembers(/* out */ Tmp, Object.Key, *Object.Value);
+	//		WriteSkippableSlice(Out, Tmp);
+	//		Tmp.Reset();
+	//	}
+
+	//	// Write object terminator
+	//	WriteSkippableSlice(Out, TConstArrayView64<uint8>());
+	//	WriteU32(Out, Magics[4]);
+	//
+	//	// Write names
+	//	WriteNumAndArray(Out, SavedNames.Names.Array());
+	//	WriteU32(Out, Magics[5]);
+	//	
+	//	return Out;
+	//}
+
+private:
+	using IdBuiltStructPair = TPair<FStructSchemaId, FBuiltStruct*>;
+	using FPropertyBindings = TCustomPropertyBindings<PropertyBatchType>;
+	static constexpr uint32 Magics[] = { 0xFEEDF00D, 0xABCD1234, 0xDADADAAA, 0x99887766, 0xF0F1F2F3 , 0x00112233};
+
+	TArray<IdBuiltStructPair>	SavedObjects;
+	PropertyBatchType			SavedProperties;
+	mutable FScratchAllocator	Scratch;
+	FPropertyBindings			FlatBindings;
+	FPropertyBindings			DeltaBindings;	
+	FSaveContext				FlatCtx;
+	FSaveContext				DeltaCtx;
+
+	//template<typename ArrayType>
+	//static void WriteNumAndArray(TArray64<uint8>& Out, const ArrayType& Items)
+	//{
+	//	WriteU32(Out, IntCastChecked<uint32>(Items.Num()));
+	//	WriteArray(Out, Items);
+	//}
+
+	//template<typename T>
+	//static TConstArrayView<T> GrabNumAndArray(/* in-out */ FByteReader& It)
+	//{
+	//	uint32 Num = It.Grab<uint32>();
+	//	return MakeArrayView(reinterpret_cast<const T*>(It.GrabBytes(Num * sizeof(T))), Num);
+	//}
+};
+
+
+static void Roundtrip()
+{
+	TBatchSaver<FMemoryPropertyBatch> Saver(GUObjectArray.GetObjectArrayNum());
+	for (TObjectIterator<UObject> It; It; ++It)
+	{
+		Saver.Save(*It);
+	}
+}
+
+static void TestBindings()
+{
+	if (FParse::Param(FCommandLine::Get(), TEXT("BindPlainProps")))
+	{
+		TScopedStructBinding<FTransform, FDefaultRuntime> Transform;
+		TScopedStructBinding<FGuid, FDefaultRuntime> Guid;
+		TScopedStructBinding<FColor, FDefaultRuntime> Color;
+		TScopedStructBinding<FLinearColor, FDefaultRuntime> LinearColor;
+		TScopedStructBinding<FFieldPath, FDefaultRuntime> FieldPath(GUE.CommonIds.FieldPath);
+		TScopedStructBinding<FScriptDelegate, FDefaultRuntime> Delegate(GUE.CommonIds.Delegate);
+		// MulticastDelegate declaration is shared with MulticastSparseDelegate
+		TScopedStructBinding<FMulticastScriptDelegate, FDefaultRuntime> InlineMulticast({GUE.CommonIds.MulticastInlineDelegate, GUE.CommonIds.MulticastDelegate});
+		InitBatchedProperties();
+		BindInitialTypes();
+		Roundtrip();
+	}
+}
+
+static FDelegateHandle GTestPreMain = FCoreDelegates::OnCommandletPreMain.AddStatic(&TestBindings);
 
 } // namespace PlainProps::UE
\ No newline at end of file
diff --git a/Engine/Plugins/Experimental/PlainPropsUObject/Source/Public/PlainPropsUObjectRuntime.h b/Engine/Plugins/Experimental/PlainPropsUObject/Source/Public/PlainPropsUObjectRuntime.h
--- a/Engine/Plugins/Experimental/PlainPropsUObject/Source/Public/PlainPropsUObjectRuntime.h
+++ b/Engine/Plugins/Experimental/PlainPropsUObject/Source/Public/PlainPropsUObjectRuntime.h
@@ -8,26 +8,37 @@
 #include "PlainPropsTypes.h"
 #include "Containers/Map.h"
 #include "UObject/NameTypes.h"
+#include "UObject/ScriptDelegateFwd.h"
+
+struct FFieldPath;
 
 namespace PlainProps::UE
 {
 
-struct FDefaultStruct;
+struct FDefaultInstance { uint64 Ptr; };
 
-class FDefaultStructs final : IDefaultStructs
+class FDefaultStructs final : public IDefaultStructs
 {
 public:
 	~FDefaultStructs();
-	void								Bind(FStructSchemaId Id, const UScriptStruct* Struct);
-	void								Drop(FStructSchemaId Id);
-	virtual const void*					Get(FStructSchemaId Id) override;
+	void									Bind(FStructSchemaId Id, const UScriptStruct* Struct);
+	void									BindZeroes(FStructSchemaId Id, SIZE_T Size, uint32 Alignment);
+	void									BindStatic(FStructSchemaId Id, const void* Struct);
+	void									Drop(FStructSchemaId Id);
+	virtual const void*						Get(FStructSchemaId Id) override;
 
 private:
-	FMutableMemoryView					Zeroes;
-	TBitArray<>							Instanced;
-	TMap<FStructSchemaId, uint8*>		Instances;
+	FMutableMemoryView						Zeroes;
+	TBitArray<>								Instanced;
+	TMap<FStructSchemaId, FDefaultInstance>	Instances;
+#if DO_CHECK
+	TBitArray<>								Bound;
+#endif
+
+	void									ReserveFlags(uint32 Idx);
 };
 
+// todo split into FCommonStructs, FCommonScopes, FCommonTypenames, FCommonMembers, access via e.g. GUE.Members.Object
 struct FCommonUObjectIds
 {
 	explicit FCommonUObjectIds(TIdIndexer<FName>& Names);
@@ -37,6 +48,7 @@ struct FCommonUObjectIds
 
 	FStructSchemaId				Name;
 	FStructSchemaId				Text;
+	FStructSchemaId				Guid;
 	FStructSchemaId				FieldPath;
 	FStructSchemaId				SoftObjectPath;
 	FStructSchemaId				ObjectPtr;
@@ -44,6 +56,10 @@ struct FCommonUObjectIds
 	FStructSchemaId				LazyObjectPtr;
 	FStructSchemaId				SoftObjectPtr;
 	FStructSchemaId				ScriptInterface;
+	FStructSchemaId				Delegate;
+	FStructSchemaId				MulticastDelegate;
+	FStructSchemaId				MulticastInlineDelegate;
+	FStructSchemaId				MulticastSparseDelegate;
 	
 	FConcreteTypenameId			Optional;
 	FConcreteTypenameId			Map;
@@ -52,7 +68,13 @@ struct FCommonUObjectIds
 	FMemberId					Key;
 	FMemberId					Value;
 	FMemberId					Add;
-	FMemberId					Del;	
+	FMemberId					Del;
+	FMemberId					Id;
+	FMemberId					Object;
+	FMemberId					Function;
+	FMemberId					Invocations;
+	FMemberId					Path;
+	FMemberId					Owner;
 };
 
 struct FGlobals
@@ -62,8 +84,8 @@ struct FGlobals
 	TIdIndexer<FName>			Names;
 	FDeclarations				Types;
 	FSchemaBindings				Schemas;
-	FCustomBindings				Customs;
-	FCustomBindings				DeltaCustoms;
+	FCustomBindingsBottom		Customs;
+	FCustomBindingsOverlay		DeltaCustoms;
 	FCommonUObjectIds			CommonIds;
 	FDefaultStructs				Defaults;
 };
@@ -89,6 +111,7 @@ struct FDefaultRuntime
 	static FDeclarations&		GetTypes()			{ return GUE.Types; }
 	static FSchemaBindings&		GetSchemas()		{ return GUE.Schemas; }
 	static FCustomBindings&		GetCustoms()		{ return GUE.Customs; }
+	static IDefaultStructs*		GetDefaults()		{ return nullptr; }
 };
 
 struct FDeltaRuntime : FDefaultRuntime
@@ -96,6 +119,54 @@ struct FDeltaRuntime : FDefaultRuntime
 	template<class T> using CustomBindings = TCustomDeltaBind<T>;
 
 	static FCustomBindings&		GetCustoms()		{ return GUE.DeltaCustoms; }
+	static IDefaultStructs*		GetDefaults()		{ return &GUE.Defaults; }
 };
 
-} // namespace PlainProps::UE
\ No newline at end of file
+
+//////////////////////////////////////////////////////////////////////////
+
+struct FFieldPathBinding : ICustomBinding
+{
+	using Type = FFieldPath;
+	const FMemberId MemberIds[2];
+
+	FFieldPathBinding(TCustomInit<FIds>) : MemberIds{GUE.CommonIds.Path, GUE.CommonIds.Owner} {}
+
+	void			Save(FMemberBuilder& Dst, const Type& Src, const Type* Default, const FSaveContext&) const;
+	void			Load(Type& Dst, FStructLoadView Src, ECustomLoadMethod Method) const;
+	static bool		Diff(const Type& A, const Type& B);
+};
+
+struct FDelegateBinding : ICustomBinding
+{
+	using Type = FScriptDelegate;
+	const FMemberId MemberIds[2];
+
+	FDelegateBinding(TCustomInit<FIds>) : MemberIds{GUE.CommonIds.Object, GUE.CommonIds.Function} {}
+
+	void			Save(FMemberBuilder& Dst, const Type& Src, const Type* Default, const FSaveContext&) const;
+	void			Load(Type& Dst, FStructLoadView Src, ECustomLoadMethod Method) const;
+	static bool		Diff(const Type& A, const Type& B);
+};
+
+struct FMulticastInlineDelegateBinding : ICustomBinding
+{
+	using Type = FMulticastScriptDelegate;
+	const FMemberId MemberIds[1];
+
+	FMulticastInlineDelegateBinding(TCustomInit<FIds>) : MemberIds{GUE.CommonIds.Invocations} {}
+
+	void			Save(FMemberBuilder& Dst, const Type& Src, const Type* Default, const FSaveContext&) const;
+	void			Load(Type& Dst, FStructLoadView Src, ECustomLoadMethod Method) const;
+	static bool		Diff(const Type& A, const Type& B);
+};
+
+
+} // namespace PlainProps::UE
+
+namespace PlainProps
+{
+template<> struct TCustomBind<FFieldPath> { using Type = UE::FFieldPathBinding; };
+template<> struct TCustomBind<FScriptDelegate> { using Type = UE::FDelegateBinding; };
+template<> struct TCustomBind<FMulticastScriptDelegate> { using Type = UE::FMulticastInlineDelegateBinding; };
+}
\ No newline at end of file
