PlainProps split CoreUObject CommonIds based on id type to avoid future name collisions

#rnx
#rnx

diff --git a/Engine/Plugins/Experimental/PlainPropsUObject/Source/Private/PlainPropsUObjectBindings.cpp b/Engine/Plugins/Experimental/PlainPropsUObject/Source/Private/PlainPropsUObjectBindings.cpp
--- a/Engine/Plugins/Experimental/PlainPropsUObject/Source/Private/PlainPropsUObjectBindings.cpp
+++ b/Engine/Plugins/Experimental/PlainPropsUObject/Source/Private/PlainPropsUObjectBindings.cpp
@@ -191,34 +191,43 @@ const void* FDefaultStructs::Get(FStructSchemaId Id)
 
 ////////////////////////////////////////////////////////////////////////////////////////////////
 
-FCommonUObjectIds::FCommonUObjectIds(TIdIndexer<FName>& Names)
+FCommonScopeIds::FCommonScopeIds(TIdIndexer<FName>& Names)
 : Core(						Names.MakeScope("/Script/Core"))
 , CoreUObject(				Names.MakeScope("/Script/CoreUObject"))
-, Name(						Names.IndexStruct({Core, Names.MakeTypename("Name")}))
-, Text(						Names.IndexStruct({Core, Names.MakeTypename("Text")}))
-, Guid(						Names.IndexStruct({Core, Names.MakeTypename("Guid")}))
-, FieldPath(				Names.IndexStruct({CoreUObject, Names.MakeTypename("FieldPath")}))
-, SoftObjectPath(			Names.IndexStruct({CoreUObject, Names.MakeTypename("SoftObjectPath")}))
-, ObjectPtr(				Names.IndexStruct({CoreUObject, Names.MakeTypename("ObjectPtr")}))
-, WeakObjectPtr(			Names.IndexStruct({CoreUObject, Names.MakeTypename("WeakObjectPtr")}))
-, LazyObjectPtr(			Names.IndexStruct({CoreUObject, Names.MakeTypename("LazyObjectPtr")}))
-, SoftObjectPtr(			Names.IndexStruct({CoreUObject, Names.MakeTypename("SoftObjectPtr")}))
-, ScriptInterface(			Names.IndexStruct({CoreUObject, Names.MakeTypename("ScriptInterface")}))
-, Delegate(					Names.IndexStruct({CoreUObject, Names.MakeTypename("Delegate")}))
-, MulticastDelegate(		Names.IndexStruct({CoreUObject, Names.MakeTypename("MulticastDelegate")}))
-, MulticastInlineDelegate(	Names.IndexStruct({CoreUObject, Names.MakeTypename("MulticastInlineDelegate")}))
-, MulticastSparseDelegate(	Names.IndexStruct({CoreUObject, Names.MakeTypename("MulticastSparseDelegate")}))
-, Optional(					Names.NameType("Optional"))
+{}
+
+FCommonTypenameIds::FCommonTypenameIds(TIdIndexer<FName>& Names)
+: Optional(					Names.NameType("Optional"))
 , Map(						Names.NameType("Map"))
 , Set(						Names.NameType("Set"))
 , Pair(						Names.NameType("Pair"))
-, Key(						Names.NameMember("Key"))
-, Value(					Names.NameMember("Value"))
+{}
+
+FCommonStructIds::FCommonStructIds(const FCommonScopeIds& Scopes, TIdIndexer<FName>& Names)
+: Name(						Names.IndexStruct({Scopes.Core, Names.MakeTypename("Name")}))
+, Text(						Names.IndexStruct({Scopes.Core, Names.MakeTypename("Text")}))
+, Guid(						Names.IndexStruct({Scopes.Core, Names.MakeTypename("Guid")}))
+, FieldPath(				Names.IndexStruct({Scopes.CoreUObject, Names.MakeTypename("FieldPath")}))
+, SoftObjectPath(			Names.IndexStruct({Scopes.CoreUObject, Names.MakeTypename("SoftObjectPath")}))
+, ObjectPtr(				Names.IndexStruct({Scopes.CoreUObject, Names.MakeTypename("ObjectPtr")}))
+, WeakObjectPtr(			Names.IndexStruct({Scopes.CoreUObject, Names.MakeTypename("WeakObjectPtr")}))
+, LazyObjectPtr(			Names.IndexStruct({Scopes.CoreUObject, Names.MakeTypename("LazyObjectPtr")}))
+, SoftObjectPtr(			Names.IndexStruct({Scopes.CoreUObject, Names.MakeTypename("SoftObjectPtr")}))
+, ScriptInterface(			Names.IndexStruct({Scopes.CoreUObject, Names.MakeTypename("ScriptInterface")}))
+, Delegate(					Names.IndexStruct({Scopes.CoreUObject, Names.MakeTypename("Delegate")}))
+, MulticastDelegate(		Names.IndexStruct({Scopes.CoreUObject, Names.MakeTypename("MulticastDelegate")}))
+, MulticastInlineDelegate(	Names.IndexStruct({Scopes.CoreUObject, Names.MakeTypename("MulticastInlineDelegate")}))
+, MulticastSparseDelegate(	Names.IndexStruct({Scopes.CoreUObject, Names.MakeTypename("MulticastSparseDelegate")}))
+{}
+
+FCommonMemberIds::FCommonMemberIds(TIdIndexer<FName>& Names)
+: Value(					Names.NameMember("Value"))
 , Add(						Names.NameMember("Add"))
 , Del(						Names.NameMember("Del"))
 , Id(						Names.NameMember("Id"))
 , Object(					Names.NameMember("Object"))
 , Function(					Names.NameMember("Function"))
+, Invocations(				Names.NameMember("Invocations"))
 , Path(						Names.NameMember("Path"))
 , Owner(					Names.NameMember("Owner"))
 {}
@@ -230,7 +239,10 @@ FGlobals::FGlobals()
 , Schemas(Names)
 , Customs(Names)
 , DeltaCustoms(Customs)
-, CommonIds(Names)
+, Scopes(Names)
+, Structs(Scopes, Names)
+, Typenames(Names)
+, Members(Names)
 {}
 
 FGlobals GUE;
@@ -251,15 +263,15 @@ static FStructSchemaId FlagsToCommonBindId(uint64 MaskedCastFlags)
 {
 	switch (MaskedCastFlags)
 	{
-		case CASTCLASS_FNameProperty:						return GUE.CommonIds.Name;
-		case CASTCLASS_FObjectProperty:						return GUE.CommonIds.ObjectPtr;
-		case CASTCLASS_FWeakObjectProperty:					return GUE.CommonIds.WeakObjectPtr;
-		case CASTCLASS_FSoftObjectProperty:					return GUE.CommonIds.SoftObjectPtr;
-		case CASTCLASS_FLazyObjectProperty:					return GUE.CommonIds.LazyObjectPtr;
-		case CASTCLASS_FDelegateProperty:					return GUE.CommonIds.Delegate;
-		case CASTCLASS_FMulticastInlineDelegateProperty:	return GUE.CommonIds.MulticastInlineDelegate;
-		case CASTCLASS_FTextProperty:						return GUE.CommonIds.Text;
-		case CASTCLASS_FFieldPathProperty:					return GUE.CommonIds.FieldPath;
+		case CASTCLASS_FNameProperty:						return GUE.Structs.Name;
+		case CASTCLASS_FObjectProperty:						return GUE.Structs.ObjectPtr;
+		case CASTCLASS_FWeakObjectProperty:					return GUE.Structs.WeakObjectPtr;
+		case CASTCLASS_FSoftObjectProperty:					return GUE.Structs.SoftObjectPtr;
+		case CASTCLASS_FLazyObjectProperty:					return GUE.Structs.LazyObjectPtr;
+		case CASTCLASS_FDelegateProperty:					return GUE.Structs.Delegate;
+		case CASTCLASS_FMulticastInlineDelegateProperty:	return GUE.Structs.MulticastInlineDelegate;
+		case CASTCLASS_FTextProperty:						return GUE.Structs.Text;
+		case CASTCLASS_FFieldPathProperty:					return GUE.Structs.FieldPath;
 		default:											break; // error
 	}
 
@@ -367,7 +379,7 @@ static FStructSchemaId DeclareStruct(const UStruct* Struct)
 
 static FTypedRange SaveNames(TConstArrayView<FName> Names, const FSaveContext& Ctx)
 {
-	FStructSchemaId Id = GUE.CommonIds.Name;
+	FStructSchemaId Id = GUE.Structs.Name;
 	FStructRangeSaver Out(Ctx.Scratch, Names.Num());
 	for (FName Name : Names)
 	{
@@ -388,8 +400,8 @@ static void LoadNames(TArray<FName>& Dst, FStructRangeLoadView Src)
 
 void FFieldPathBinding::Save(FMemberBuilder& Dst, const FFieldPath& Src, const FFieldPath*, const FSaveContext& Ctx) const
 {
-	Dst.AddRange(GUE.CommonIds.Path, SaveNames(Src.*GFieldPathPath, Ctx));
-	Dst.AddStruct(GUE.CommonIds.Owner, GUE.CommonIds.WeakObjectPtr, SaveStruct(&(Src.*GFieldPathOwner), GUE.CommonIds.WeakObjectPtr, Ctx));
+	Dst.AddRange(GUE.Members.Path, SaveNames(Src.*GFieldPathPath, Ctx));
+	Dst.AddStruct(GUE.Members.Owner, GUE.Structs.WeakObjectPtr, SaveStruct(&(Src.*GFieldPathOwner), GUE.Structs.WeakObjectPtr, Ctx));
 }
 
 void FFieldPathBinding::Load(FFieldPath& Dst, FStructLoadView Src, ECustomLoadMethod Method) const
@@ -411,8 +423,8 @@ void FDelegateBinding::Save(FMemberBuilder& Dst, const FScriptDelegate& Src, con
 {
 	if (Src.IsBound())
 	{
-		Dst.AddStruct(GUE.CommonIds.Object, GUE.CommonIds.WeakObjectPtr, SaveStruct(&(Src.*GDelegateObject), GUE.CommonIds.WeakObjectPtr, Ctx));
-		Dst.AddStruct(GUE.CommonIds.Function, GUE.CommonIds.Name, SaveStruct(&(Src.*GDelegateFunctionName), GUE.CommonIds.Name, Ctx));
+		Dst.AddStruct(GUE.Members.Object, GUE.Structs.WeakObjectPtr, SaveStruct(&(Src.*GDelegateObject), GUE.Structs.WeakObjectPtr, Ctx));
+		Dst.AddStruct(GUE.Members.Function, GUE.Structs.Name, SaveStruct(&(Src.*GDelegateFunctionName), GUE.Structs.Name, Ctx));
 	}
 }
 
@@ -439,8 +451,8 @@ bool FDelegateBinding::Diff(const FScriptDelegate& A, const FScriptDelegate& B)
 
 static void SaveUnicastDelegate(FMemberBuilder& Dst, const FUnicastScriptDelegate& Src, const FSaveContext& Ctx)
 {
-	Dst.AddStruct(GUE.CommonIds.Object, GUE.CommonIds.WeakObjectPtr, SaveStruct(&(Src.*GUnicastDelegateObject), GUE.CommonIds.WeakObjectPtr, Ctx));
-	Dst.AddStruct(GUE.CommonIds.Function, GUE.CommonIds.Name, SaveStruct(&(Src.*GUnicastDelegateFunctionName), GUE.CommonIds.Name, Ctx));
+	Dst.AddStruct(GUE.Members.Object, GUE.Structs.WeakObjectPtr, SaveStruct(&(Src.*GUnicastDelegateObject), GUE.Structs.WeakObjectPtr, Ctx));
+	Dst.AddStruct(GUE.Members.Function, GUE.Structs.Name, SaveStruct(&(Src.*GUnicastDelegateFunctionName), GUE.Structs.Name, Ctx));
 }
 
 static void LoadUnicastDelegate(FUnicastScriptDelegate& Dst, FStructLoadView Src)
@@ -452,7 +464,7 @@ static void LoadUnicastDelegate(FUnicastScriptDelegate& Dst, FStructLoadView Src
 
 static FTypedRange SaveInvocations(const FMulticastInvocationList& In, const FSaveContext& Ctx)
 {
-	FStructSchemaId ItemId = GUE.CommonIds.Delegate;
+	FStructSchemaId ItemId = GUE.Structs.Delegate;
 	FMemberSchema Schema = MakeStructRangeSchema(DefaultRangeMax, ItemId);
 	if (int32 NumTotal = In.Num())
 	{
@@ -486,12 +498,12 @@ static FTypedRange SaveInvocations(const FMulticastInvocationList& In, const FSa
 
 static void SaveMulticastDelegate(FMemberBuilder& Dst, const FMulticastScriptDelegate& Src, const FSaveContext& Ctx)
 {
-	Dst.AddRange(GUE.CommonIds.Invocations, SaveInvocations(Src.*GMulticastDelegateInvocationList, Ctx));
+	Dst.AddRange(GUE.Members.Invocations, SaveInvocations(Src.*GMulticastDelegateInvocationList, Ctx));
 }
 
 static void SaveEmptyMulticastDelegate(FMemberBuilder& Dst)
 {
-	Dst.AddRange(GUE.CommonIds.Invocations, { MakeStructRangeSchema(DefaultRangeMax, GUE.CommonIds.Delegate), nullptr });
+	Dst.AddRange(GUE.Members.Invocations, { MakeStructRangeSchema(DefaultRangeMax, GUE.Structs.Delegate), nullptr });
 }
 
 static void LoadInvocations(FMulticastInvocationList& Dst, FStructRangeLoadView Src)
@@ -645,13 +657,13 @@ static FStructSchemaId BindSparseDelegate(FStructSchemaId Owner, FMulticastSpars
 	// Todo: Ownership / memory leak
 	ICustomBinding* Leak = new FMulticastSparseDelegateBinding(CastChecked<USparseDelegateFunction>(Property->SignatureFunction));
 
-	FTypeId MulticastSparseDelegate = GUE.Names.Resolve(GUE.CommonIds.MulticastSparseDelegate);
+	FTypeId MulticastSparseDelegate = GUE.Names.Resolve(GUE.Structs.MulticastSparseDelegate);
 	FTypeId OwnerParam = GUE.Names.Resolve(Owner);
-	FTypeId PropertyParam = { GUE.CommonIds.CoreUObject, FTypenameId(GUE.Names.NameType(Property->GetFName())) };
+	FTypeId PropertyParam = { GUE.Scopes.CoreUObject, FTypenameId(GUE.Names.NameType(Property->GetFName())) };
 	FTypeId UniqueBindName = GUE.Names.MakeParametricType(MulticastSparseDelegate, {OwnerParam, PropertyParam});
 	FStructSchemaId BindId = GUE.Names.IndexStruct(UniqueBindName);
 
-	GUE.Customs.BindStruct(BindId, *Leak, GUE.Types.Get(GUE.CommonIds.MulticastDelegate));
+	GUE.Customs.BindStruct(BindId, *Leak, GUE.Types.Get(GUE.Structs.MulticastDelegate));
 
 	return BindId;
 }
@@ -685,7 +697,7 @@ struct FInterfaceBinding final : ICustomBinding
 	
 	void Save(FMemberBuilder& Dst, const FScriptInterface& Src, const FSaveContext& Ctx) const
 	{
-		Dst.AddStruct(GUE.CommonIds.Object, GUE.CommonIds.ObjectPtr, SaveStruct(Src.GetObject(), GUE.CommonIds.ObjectPtr, Ctx));
+		Dst.AddStruct(GUE.Members.Object, GUE.Structs.ObjectPtr, SaveStruct(Src.GetObject(), GUE.Structs.ObjectPtr, Ctx));
 	}
 
 	void Load(FScriptInterface& Dst, FStructLoadView Src) const
@@ -705,8 +717,8 @@ class FInterfaceBindings
 
 public:
 	FInterfaceBindings()
-	: ScriptInterface(GUE.Names.Resolve(GUE.CommonIds.ScriptInterface))
-	, Declaration(GUE.Types.DeclareStruct(GUE.CommonIds.ScriptInterface, ScriptInterface, {GUE.CommonIds.Object}, EMemberPresence::RequireAll))
+	: ScriptInterface(GUE.Names.Resolve(GUE.Structs.ScriptInterface))
+	, Declaration(GUE.Types.DeclareStruct(GUE.Structs.ScriptInterface, ScriptInterface, {GUE.Members.Object}, EMemberPresence::RequireAll))
 	{}
 
 	FStructSchemaId Bind(FInterfaceProperty* Property)
@@ -1604,7 +1616,7 @@ void LoadSetDelta(const BindingType& Binding, ScriptType& Dst, FStructLoadView S
 	FMemberLoader Members(Src);				
 	while (Members.HasMore())
 	{
-		if (Members.PeekName() == GUE.CommonIds.Add)
+		if (Members.PeekName() == GUE.Members.Add)
 		{
 			if (Dst.IsEmpty())
 			{
@@ -1618,7 +1630,7 @@ void LoadSetDelta(const BindingType& Binding, ScriptType& Dst, FStructLoadView S
 			break;
 		}
 
-		check(Members.PeekName() == GUE.CommonIds.Del);
+		check(Members.PeekName() == GUE.Members.Del);
 		FRangeLoadView Keys = Members.GrabRange();
 		if (!Dst.IsEmpty())
 		{
@@ -1757,7 +1769,7 @@ struct TSetPropertyBinding : IItemRangeBinding, ICustomBinding
 		{
 			if (Default && !Default->IsEmpty())
 			{
-				Dst.AddRange(GUE.CommonIds.Del, SaveAll(*Default, Layout, Range, Ctx));
+				Dst.AddRange(GUE.Members.Del, SaveAll(*Default, Layout, Range, Ctx));
 			}
 		}
 		else if (Default && !Default->IsEmpty())
@@ -1769,7 +1781,7 @@ struct TSetPropertyBinding : IItemRangeBinding, ICustomBinding
 			}
 			if (FSubSetIterator Deleted{Layout, *Default, DeletedIds})
 			{
-				Dst.AddRange(GUE.CommonIds.Del, SaveSome(Deleted, Range, Ctx));
+				Dst.AddRange(GUE.Members.Del, SaveSome(Deleted, Range, Ctx));
 			}
 			
 			TBitArray<> AddedIds(false, Src.GetMaxIndex());
@@ -1779,12 +1791,12 @@ struct TSetPropertyBinding : IItemRangeBinding, ICustomBinding
 			}
 			if (FSubSetIterator Added{Layout, Src, AddedIds})
 			{
-				Dst.AddRange(GUE.CommonIds.Add, SaveSome(Added, Range, Ctx));
+				Dst.AddRange(GUE.Members.Add, SaveSome(Added, Range, Ctx));
 			}
 		}
 		else
 		{
-			Dst.AddRange(GUE.CommonIds.Add, SaveAll(Src, Layout, Range, Ctx));
+			Dst.AddRange(GUE.Members.Add, SaveAll(Src, Layout, Range, Ctx));
 		}
 	}
 
@@ -1979,9 +1991,9 @@ public:
 
 		// Index custom delta binding struct name
 		FTypeId Params[1] = { Elem.IndexParameterName(GUE.Names) };
-		FTypeId Type = { GUE.CommonIds.Core, FTypenameId(GUE.Names.MakeParametricTypeId(GUE.CommonIds.Set, Params)) };
+		FTypeId Type = { GUE.Scopes.Core, FTypenameId(GUE.Names.MakeParametricTypeId(GUE.Typenames.Set, Params)) };
 		FStructSchemaId Id = GUE.Names.IndexStruct(Type);
-		FMemberId Members[2] = { GUE.CommonIds.Del, GUE.CommonIds.Add };
+		FMemberId Members[2] = { GUE.Members.Del, GUE.Members.Add };
 
 		// Todo: Ownership / memory leak
 		const FStructDeclaration& Declaration = GUE.Types.DeclareStruct(Id, Type, Members, EMemberPresence::AllowSparse);
@@ -2167,7 +2179,7 @@ struct TMapPropertyDeltaBinding : TMapPropertyBinding<FScriptMap, KeyKind, Value
 		{
 			if (Default && !Default->IsEmpty())
 			{
-				Dst.AddRange(GUE.CommonIds.Del, SaveAll(*Default, Layout, KeyRange, Ctx));
+				Dst.AddRange(GUE.Members.Del, SaveAll(*Default, Layout, KeyRange, Ctx));
 			}
 		}
 		else if (Default && !Default->IsEmpty())
@@ -2179,7 +2191,7 @@ struct TMapPropertyDeltaBinding : TMapPropertyBinding<FScriptMap, KeyKind, Value
 			}
 			if (FSubMapIterator Deleted{Layout, *Default, DeletedIds})
 			{
-				Dst.AddRange(GUE.CommonIds.Del, SaveSome(Deleted, KeyRange, Ctx));
+				Dst.AddRange(GUE.Members.Del, SaveSome(Deleted, KeyRange, Ctx));
 			}
 			
 			TBitArray<> AddedIds(false, Src.GetMaxIndex());
@@ -2189,12 +2201,12 @@ struct TMapPropertyDeltaBinding : TMapPropertyBinding<FScriptMap, KeyKind, Value
 			}
 			if (FSubMapIterator Added{Layout, Src, AddedIds})
 			{
-				Dst.AddRange(GUE.CommonIds.Add, SaveSome(Added, PairRange, Ctx));
+				Dst.AddRange(GUE.Members.Add, SaveSome(Added, PairRange, Ctx));
 			}
 		}
 		else
 		{
-			Dst.AddRange(GUE.CommonIds.Add, SaveAll(Src, Layout, PairRange, Ctx));
+			Dst.AddRange(GUE.Members.Add, SaveAll(Src, Layout, PairRange, Ctx));
 		}
 	}
 
@@ -2323,11 +2335,11 @@ class FMapBindings
 		// Index custom delta binding struct name
 		FParametricTypeId PairTypename = GUE.Names.Resolve(static_cast<FStructSchemaId>(Members.Pair.InnermostSchema.Get())).Name.AsParametric();
 		TConstArrayView<FTypeId> Params = GUE.Names.Resolve(PairTypename).GetParameters();
-		FTypeId Type = { GUE.CommonIds.Core, FTypenameId(GUE.Names.MakeParametricTypeId(GUE.CommonIds.Map, Params)) };
+		FTypeId Type = { GUE.Scopes.Core, FTypenameId(GUE.Names.MakeParametricTypeId(GUE.Typenames.Map, Params)) };
 		FStructSchemaId Id = GUE.Names.IndexStruct(Type);
 
 		// Todo: Ownership / memory leak
-		const FStructDeclaration& Declaration = GUE.Types.DeclareStruct(Id, Type, { GUE.CommonIds.Del, GUE.CommonIds.Add }, EMemberPresence::AllowSparse);
+		const FStructDeclaration& Declaration = GUE.Types.DeclareStruct(Id, Type, { GUE.Members.Del, GUE.Members.Add }, EMemberPresence::AllowSparse);
 		auto Out = new TMapPropertyDeltaBinding<KeyKind, ValueKind>(Property, Members);
 		GUE.DeltaCustoms.BindStruct(Declaration.Id, *Out, Declaration, {});
 
@@ -2410,9 +2422,9 @@ class FPairBindings
 
 		// Index names
 		FTypeId Params[2] = { Pair.KV[0].IndexParameterName(GUE.Names), Pair.KV[1].IndexParameterName(GUE.Names) };
-		FTypeId Type = { GUE.CommonIds.Core, FTypenameId(GUE.Names.MakeParametricTypeId(GUE.CommonIds.Pair, Params)) };
+		FTypeId Type = { GUE.Scopes.Core, FTypenameId(GUE.Names.MakeParametricTypeId(GUE.Typenames.Pair, Params)) };
 		FStructSchemaId Id = GUE.Names.IndexStruct(Type);
-		FMemberId Members[2] = { GUE.CommonIds.Key, GUE.CommonIds.Value };
+		FMemberId Members[2] = { GUE.Members.Key, GUE.Members.Value };
 
 		// Todo: Ownership / memory leak
 		GUE.Types.DeclareStruct(Id, Type, Members, EMemberPresence::RequireAll);
@@ -2450,9 +2462,9 @@ class FOptionalBindings
 
 		// Index names
 		FTypeId Params[1] = { Value.IndexParameterName(GUE.Names) };
-		FTypeId Type = { GUE.CommonIds.Core, FTypenameId(GUE.Names.MakeParametricTypeId(GUE.CommonIds.Optional, Params)) };
+		FTypeId Type = { GUE.Scopes.Core, FTypenameId(GUE.Names.MakeParametricTypeId(GUE.Typenames.Optional, Params)) };
 		FStructSchemaId Id = GUE.Names.IndexStruct(Type);
-		FMemberId Members[1] = { GUE.CommonIds.Value };
+		FMemberId Members[1] = { GUE.Members.Value };
 		FMemberBinding Inners[1] = { Value };
 
 		// Todo: Ownership / memory leak
@@ -2876,19 +2888,19 @@ static void DeclareProperty(FStructSchemaId Id, TConstArrayView<FMemberId> Membe
 static void InitBatchedProperties()
 {
 	// All property batch types share declaration, but some members could be exclusive to a specific type e.g. FMemoryPropertyBatch
-	DeclareProperty(GUE.CommonIds.Name,				{ GUE.CommonIds.Id }, EMemberPresence::RequireAll);
-	DeclareProperty(GUE.CommonIds.Text,				{ GUE.CommonIds.Id }, EMemberPresence::RequireAll);
-	DeclareProperty(GUE.CommonIds.ObjectPtr,		{ GUE.CommonIds.Id }, EMemberPresence::RequireAll);
-	DeclareProperty(GUE.CommonIds.WeakObjectPtr,	{ GUE.CommonIds.Id }, EMemberPresence::RequireAll);
-	DeclareProperty(GUE.CommonIds.SoftObjectPtr,	{ GUE.CommonIds.Id }, EMemberPresence::RequireAll);
-	DeclareProperty(GUE.CommonIds.LazyObjectPtr,	{ GUE.CommonIds.Id }, EMemberPresence::RequireAll);
+	DeclareProperty(GUE.Structs.Name,			{ GUE.Members.Id }, EMemberPresence::RequireAll);
+	DeclareProperty(GUE.Structs.Text,			{ GUE.Members.Id }, EMemberPresence::RequireAll);
+	DeclareProperty(GUE.Structs.ObjectPtr,		{ GUE.Members.Id }, EMemberPresence::RequireAll);
+	DeclareProperty(GUE.Structs.WeakObjectPtr,	{ GUE.Members.Id }, EMemberPresence::RequireAll);
+	DeclareProperty(GUE.Structs.SoftObjectPtr,	{ GUE.Members.Id }, EMemberPresence::RequireAll);
+	DeclareProperty(GUE.Structs.LazyObjectPtr,	{ GUE.Members.Id }, EMemberPresence::RequireAll);
 
-	GUE.Defaults.BindZeroes(GUE.CommonIds.Name, sizeof(FName), alignof(FName));
-	GUE.Defaults.BindStatic(GUE.CommonIds.Text, &FText::GetEmpty());
-	GUE.Defaults.BindZeroes(GUE.CommonIds.ObjectPtr, sizeof(FObjectPtr), alignof(FObjectPtr));
-	GUE.Defaults.BindZeroes(GUE.CommonIds.WeakObjectPtr, sizeof(FWeakObjectPtr), alignof(FWeakObjectPtr));
-	GUE.Defaults.BindZeroes(GUE.CommonIds.SoftObjectPtr, sizeof(FSoftObjectPtr), alignof(FSoftObjectPtr));
-	GUE.Defaults.BindZeroes(GUE.CommonIds.LazyObjectPtr, sizeof(FLazyObjectPtr), alignof(FLazyObjectPtr));
+	GUE.Defaults.BindZeroes(GUE.Structs.Name, sizeof(FName), alignof(FName));
+	GUE.Defaults.BindStatic(GUE.Structs.Text, &FText::GetEmpty());
+	GUE.Defaults.BindZeroes(GUE.Structs.ObjectPtr, sizeof(FObjectPtr), alignof(FObjectPtr));
+	GUE.Defaults.BindZeroes(GUE.Structs.WeakObjectPtr, sizeof(FWeakObjectPtr), alignof(FWeakObjectPtr));
+	GUE.Defaults.BindZeroes(GUE.Structs.SoftObjectPtr, sizeof(FSoftObjectPtr), alignof(FSoftObjectPtr));
+	GUE.Defaults.BindZeroes(GUE.Structs.LazyObjectPtr, sizeof(FLazyObjectPtr), alignof(FLazyObjectPtr));
 }
 
 struct FMemoryPropertyBatch
@@ -2897,7 +2909,7 @@ struct FMemoryPropertyBatch
 
 	static void Save(FMemberBuilder& Out, FName In, const FSaveContext&)
 	{
-		Out.Add(GUE.CommonIds.Id, In.ToUnstableInt());
+		Out.Add(GUE.Members.Id, In.ToUnstableInt());
 	}
 
 	static void Load(FName& Out, FStructLoadView In)
@@ -2909,7 +2921,7 @@ struct FMemoryPropertyBatch
 
 	void Save(FMemberBuilder& Out, const FText& In, const FSaveContext&)
 	{
-		Out.Add(GUE.CommonIds.Id, Texts.Num());
+		Out.Add(GUE.Members.Id, Texts.Num());
 		Texts.Add(In);
 	}
 
@@ -2924,7 +2936,7 @@ struct FMemoryPropertyBatch
 	static void Save(FMemberBuilder& Out, FObjectHandle In, const FSaveContext&)
 	{
 		static_assert(sizeof(In) == sizeof(uint64));
-		Out.Add(GUE.CommonIds.Id, reinterpret_cast<uint64&>(In));
+		Out.Add(GUE.Members.Id, reinterpret_cast<uint64&>(In));
 	}
 
 	static void Load(FObjectHandle& Out, FStructLoadView In)
@@ -2938,7 +2950,7 @@ struct FMemoryPropertyBatch
 	{
 		// Save ObjectSerialNumber + ObjectIndex a single uint64
 		static_assert(sizeof(FWeakObjectPtr) == sizeof(uint64));
-		Out.Add(GUE.CommonIds.Id, reinterpret_cast<const uint64&>(In));
+		Out.Add(GUE.Members.Id, reinterpret_cast<const uint64&>(In));
 	}
 
 	static void Load(FWeakObjectPtr& Out, FStructLoadView In)
@@ -2950,8 +2962,8 @@ struct FMemoryPropertyBatch
 
 	static void Save(FMemberBuilder& Out, const FSoftObjectPtr& In, const FSaveContext& Ctx)
 	{
-		FBuiltStruct* SoftPath = SaveStruct(&In.GetUniqueID(), GUE.CommonIds.SoftObjectPath, Ctx);
-		Out.AddStruct(GUE.CommonIds.Id, GUE.CommonIds.SoftObjectPath, SoftPath);
+		FBuiltStruct* SoftPath = SaveStruct(&In.GetUniqueID(), GUE.Structs.SoftObjectPath, Ctx);
+		Out.AddStruct(GUE.Members.Id, GUE.Structs.SoftObjectPath, SoftPath);
 	}
 
 	static void Load(FSoftObjectPtr& Out, FStructLoadView In)
@@ -2963,8 +2975,8 @@ struct FMemoryPropertyBatch
 
 	static void Save(FMemberBuilder& Out, const FLazyObjectPtr& In, const FSaveContext& Ctx)
 	{
-		FBuiltStruct* Guid = SaveStruct(&In.GetUniqueID(), GUE.CommonIds.Guid, Ctx);
-		Out.AddStruct(GUE.CommonIds.Id, GUE.CommonIds.Guid, Guid);
+		FBuiltStruct* Guid = SaveStruct(&In.GetUniqueID(), GUE.Structs.Guid, Ctx);
+		Out.AddStruct(GUE.Members.Id, GUE.Structs.Guid, Guid);
 	}
 
 	static void Load(FLazyObjectPtr& Out, FStructLoadView In)
@@ -3041,12 +3053,12 @@ struct TCustomPropertyBindings
 	, WeakObjectPtr(Batch)
 	, LazyObjectPtr(Batch)
 	{
-		Bind(GUE.CommonIds.Name, Name);
-		Bind(GUE.CommonIds.Text, Text);
-		Bind(GUE.CommonIds.ObjectPtr, ObjectPtr);
-		Bind(GUE.CommonIds.SoftObjectPtr, SoftObjectPtr);
-		Bind(GUE.CommonIds.WeakObjectPtr, WeakObjectPtr);
-		Bind(GUE.CommonIds.LazyObjectPtr, LazyObjectPtr);
+		Bind(GUE.Structs.Name, Name);
+		Bind(GUE.Structs.Text, Text);
+		Bind(GUE.Structs.ObjectPtr, ObjectPtr);
+		Bind(GUE.Structs.SoftObjectPtr, SoftObjectPtr);
+		Bind(GUE.Structs.WeakObjectPtr, WeakObjectPtr);
+		Bind(GUE.Structs.LazyObjectPtr, LazyObjectPtr);
 	}
 
 	void Bind(FStructSchemaId Id, ICustomBinding& Binding)
@@ -3186,10 +3198,10 @@ static void TestBindings()
 		TScopedStructBinding<FGuid, FDefaultRuntime> Guid;
 		TScopedStructBinding<FColor, FDefaultRuntime> Color;
 		TScopedStructBinding<FLinearColor, FDefaultRuntime> LinearColor;
-		TScopedStructBinding<FFieldPath, FDefaultRuntime> FieldPath(GUE.CommonIds.FieldPath);
-		TScopedStructBinding<FScriptDelegate, FDefaultRuntime> Delegate(GUE.CommonIds.Delegate);
+		TScopedStructBinding<FFieldPath, FDefaultRuntime> FieldPath(GUE.Structs.FieldPath);
+		TScopedStructBinding<FScriptDelegate, FDefaultRuntime> Delegate(GUE.Structs.Delegate);
 		// MulticastDelegate declaration is shared with MulticastSparseDelegate
-		TScopedStructBinding<FMulticastScriptDelegate, FDefaultRuntime> InlineMulticast({GUE.CommonIds.MulticastInlineDelegate, GUE.CommonIds.MulticastDelegate});
+		TScopedStructBinding<FMulticastScriptDelegate, FDefaultRuntime> InlineMulticast({GUE.Structs.MulticastInlineDelegate, GUE.Structs.MulticastDelegate});
 		InitBatchedProperties();
 		BindInitialTypes();
 		Roundtrip();
diff --git a/Engine/Plugins/Experimental/PlainPropsUObject/Source/Public/PlainPropsUObjectRuntime.h b/Engine/Plugins/Experimental/PlainPropsUObject/Source/Public/PlainPropsUObjectRuntime.h
--- a/Engine/Plugins/Experimental/PlainPropsUObject/Source/Public/PlainPropsUObjectRuntime.h
+++ b/Engine/Plugins/Experimental/PlainPropsUObject/Source/Public/PlainPropsUObjectRuntime.h
@@ -38,14 +38,28 @@ private:
 	void									ReserveFlags(uint32 Idx);
 };
 
-// todo split into FCommonStructs, FCommonScopes, FCommonTypenames, FCommonMembers, access via e.g. GUE.Members.Object
-struct FCommonUObjectIds
+struct FCommonScopeIds
 {
-	explicit FCommonUObjectIds(TIdIndexer<FName>& Names);
+	explicit FCommonScopeIds(TIdIndexer<FName>& Names);
 	
 	FScopeId					Core;
 	FScopeId					CoreUObject;
+};
+
+struct FCommonTypenameIds
+{
+	explicit FCommonTypenameIds(TIdIndexer<FName>& Names);
+	
+	FConcreteTypenameId			Optional;
+	FConcreteTypenameId			Map;
+	FConcreteTypenameId			Set;
+	FConcreteTypenameId			Pair;
+};
 
+struct FCommonStructIds
+{
+	explicit FCommonStructIds(const FCommonScopeIds& Scopes, TIdIndexer<FName>& Names);
+	
 	FStructSchemaId				Name;
 	FStructSchemaId				Text;
 	FStructSchemaId				Guid;
@@ -60,11 +74,12 @@ struct FCommonUObjectIds
 	FStructSchemaId				MulticastDelegate;
 	FStructSchemaId				MulticastInlineDelegate;
 	FStructSchemaId				MulticastSparseDelegate;
+};
+
+struct FCommonMemberIds
+{
+	explicit FCommonMemberIds(TIdIndexer<FName>& Names);
 	
-	FConcreteTypenameId			Optional;
-	FConcreteTypenameId			Map;
-	FConcreteTypenameId			Set;
-	FConcreteTypenameId			Pair;
 	FMemberId					Key;
 	FMemberId					Value;
 	FMemberId					Add;
@@ -86,8 +101,11 @@ struct FGlobals
 	FSchemaBindings				Schemas;
 	FCustomBindingsBottom		Customs;
 	FCustomBindingsOverlay		DeltaCustoms;
-	FCommonUObjectIds			CommonIds;
 	FDefaultStructs				Defaults;
+	FCommonScopeIds				Scopes;
+	FCommonStructIds			Structs;
+	FCommonTypenameIds			Typenames;
+	FCommonMemberIds			Members;
 };
 PLAINPROPSUOBJECT_API extern FGlobals GUE;
 
@@ -130,7 +148,7 @@ struct FFieldPathBinding : ICustomBinding
 	using Type = FFieldPath;
 	const FMemberId MemberIds[2];
 
-	FFieldPathBinding(TCustomInit<FIds>) : MemberIds{GUE.CommonIds.Path, GUE.CommonIds.Owner} {}
+	FFieldPathBinding(TCustomInit<FIds>) : MemberIds{GUE.Members.Path, GUE.Members.Owner} {}
 
 	void			Save(FMemberBuilder& Dst, const Type& Src, const Type* Default, const FSaveContext&) const;
 	void			Load(Type& Dst, FStructLoadView Src, ECustomLoadMethod Method) const;
@@ -142,7 +160,7 @@ struct FDelegateBinding : ICustomBinding
 	using Type = FScriptDelegate;
 	const FMemberId MemberIds[2];
 
-	FDelegateBinding(TCustomInit<FIds>) : MemberIds{GUE.CommonIds.Object, GUE.CommonIds.Function} {}
+	FDelegateBinding(TCustomInit<FIds>) : MemberIds{GUE.Members.Object, GUE.Members.Function} {}
 
 	void			Save(FMemberBuilder& Dst, const Type& Src, const Type* Default, const FSaveContext&) const;
 	void			Load(Type& Dst, FStructLoadView Src, ECustomLoadMethod Method) const;
@@ -154,7 +172,7 @@ struct FMulticastInlineDelegateBinding : ICustomBinding
 	using Type = FMulticastScriptDelegate;
 	const FMemberId MemberIds[1];
 
-	FMulticastInlineDelegateBinding(TCustomInit<FIds>) : MemberIds{GUE.CommonIds.Invocations} {}
+	FMulticastInlineDelegateBinding(TCustomInit<FIds>) : MemberIds{GUE.Members.Invocations} {}
 
 	void			Save(FMemberBuilder& Dst, const Type& Src, const Type* Default, const FSaveContext&) const;
 	void			Load(Type& Dst, FStructLoadView Src, ECustomLoadMethod Method) const;
