PlainProps faster APIs to load struct with a single member

#rnx

diff --git a/Engine/Plugins/Experimental/PlainProps/Source/Private/PlainPropsLoad.cpp b/Engine/Plugins/Experimental/PlainProps/Source/Private/PlainPropsLoad.cpp
--- a/Engine/Plugins/Experimental/PlainProps/Source/Private/PlainPropsLoad.cpp
+++ b/Engine/Plugins/Experimental/PlainProps/Source/Private/PlainPropsLoad.cpp
@@ -1148,4 +1148,11 @@ void ConstructAndLoadStruct(void* Dst, FStructLoadView Src)
 	ConstructAndLoadStruct(Dst, Src.Values, Src.Schema.LoadId, Src.Schema.Batch);
 }
 
+void LoadSoleStruct(void* Dst, FStructLoadView Src)
+{
+	// Todo: Optimize
+	LoadStruct(Dst, FMemberLoader(Src).GrabStruct());
+	//FLoadStructPlan Plan = Src.Schema.Batch[Src.Schema.LoadId];
+}
+
 } // namespace PlainProps
\ No newline at end of file
diff --git a/Engine/Plugins/Experimental/PlainProps/Source/Public/PlainPropsLoadMember.h b/Engine/Plugins/Experimental/PlainProps/Source/Public/PlainPropsLoadMember.h
--- a/Engine/Plugins/Experimental/PlainProps/Source/Public/PlainPropsLoadMember.h
+++ b/Engine/Plugins/Experimental/PlainProps/Source/Public/PlainPropsLoadMember.h
@@ -29,6 +29,35 @@ struct FStructLoadView
 PLAINPROPS_API void	LoadStruct(void* Dst, FStructLoadView Src);
 PLAINPROPS_API void	ConstructAndLoadStruct(void* Dst, FStructLoadView Src);
 
+// Load single struct member faster than LoadStruct(Dst, FMemberLoader(Src).GrabStruct())
+PLAINPROPS_API void LoadSoleStruct(void* Dst, FStructLoadView Src);
+
+// Load single leaf member faster than FMemberLoader(Src).GrabLeaf().As<T>()
+template<LeafType T>
+inline void LoadSole(void* Dst, FStructLoadView Src)
+{
+	Src.Values.CheckSize(sizeof(T));
+	checkSlow(ReflectLeaf<T> == SizeOf(FMemberLoader(Src).GrabLeaf().Leaf));
+	FMemory::Memcpy(Dst, Src.Values.Peek(), sizeof(T));
+}
+
+template<>
+inline void LoadSole<bool>(void* Dst, FStructLoadView Src)
+{
+	Src.Values.CheckSize(sizeof(bool));
+	checkSlow(ReflectLeaf<T> == SizeOf(FMemberLoader(Src).GrabLeaf().Leaf));
+	FBitCacheReader Bits;
+	*static_cast<bool*>(Dst) = Bits.GrabNext(Src.Values);
+}
+
+template<LeafType T>
+inline T LoadSole(FStructLoadView Src)
+{
+	T Out;
+	LoadSole<T>(&Out, Src);
+	return Out;
+}
+
 //////////////////////////////////////////////////////////////////////////
 
 struct FRangeLoadSchema
diff --git a/Engine/Plugins/Experimental/PlainPropsUObject/Source/Private/PlainPropsUObjectBindings.cpp b/Engine/Plugins/Experimental/PlainPropsUObject/Source/Private/PlainPropsUObjectBindings.cpp
--- a/Engine/Plugins/Experimental/PlainPropsUObject/Source/Private/PlainPropsUObjectBindings.cpp
+++ b/Engine/Plugins/Experimental/PlainPropsUObject/Source/Private/PlainPropsUObjectBindings.cpp
@@ -909,9 +909,7 @@ struct FArrayPropertyInfo
 	, bConstructor(HasConstructor(Property->Inner))
 	, ItemAlign(Property->Inner->GetMinAlignment())
 	, ItemSize(Property->Inner->GetElementSize())
-	{
-		check(ItemAlign == Property->Inner->GetMinAlignment()); // todo remove
-	}
+	{}
 
 	union
 	{
@@ -2933,9 +2931,7 @@ struct FMemoryPropertyBatch
 
 	static void Load(FName& Out, FStructLoadView In)
 	{
-		// Todo: embed in some experimental load as leaf API
-		In.Values.CheckSize(sizeof(decltype(ToInt(FName()))));
-		Out = FromInt(FPlatformMemory::ReadUnaligned<uint64>(In.Values.Peek()));
+		Out = FromInt(LoadSole<uint64>(In));
 	}
 
 	void Save(FMemberBuilder& Out, const FText& In, const FSaveContext&)
@@ -2946,23 +2942,18 @@ struct FMemoryPropertyBatch
 
 	void Load(FText& Out, FStructLoadView In) const
 	{
-		// Todo: embed in some experimental load as leaf API
-		In.Values.CheckSize(sizeof(int32));
-		int32 Idx = FPlatformMemory::ReadUnaligned<int32>(In.Values.Peek());
-		Out = Texts[Idx];
+		Out = Texts[LoadSole<int32>(In)];
 	}
 
 	static void Save(FMemberBuilder& Out, FObjectHandle In, const FSaveContext&)
 	{
 		static_assert(sizeof(In) == sizeof(uint64));
-		Out.Add(GUE.Members.Id, reinterpret_cast<uint64&>(In));
+		Out.Add(GUE.Members.Id, reinterpret_cast<const uint64&>(In));
 	}
 
 	static void Load(FObjectHandle& Out, FStructLoadView In)
 	{
-		// Todo: embed in some experimental load as leaf API
-		In.Values.CheckSize(sizeof(uint64));
-		reinterpret_cast<uint64&>(Out) = FPlatformMemory::ReadUnaligned<uint64>(In.Values.Peek());
+		LoadSole<uint64>(&Out, In);
 	}
 
 	static void Save(FMemberBuilder& Out, const FWeakObjectPtr& In, const FSaveContext&)
@@ -2974,9 +2965,7 @@ struct FMemoryPropertyBatch
 
 	static void Load(FWeakObjectPtr& Out, FStructLoadView In)
 	{
-		// Todo: embed in some experimental load as leaf API
-		In.Values.CheckSize(sizeof(uint64));
-		reinterpret_cast<uint64&>(Out) = FPlatformMemory::ReadUnaligned<uint64>(In.Values.Peek());
+		LoadSole<uint64>(&Out, In);
 	}
 
 	static void Save(FMemberBuilder& Out, const FSoftObjectPtr& In, const FSaveContext& Ctx)
@@ -2988,8 +2977,7 @@ struct FMemoryPropertyBatch
 	static void Load(FSoftObjectPtr& Out, FStructLoadView In)
 	{
 		Out.ResetWeakPtr();
-		FMemberLoader Members(In);
-		LoadStruct(&Out.GetUniqueID(), Members.GrabStruct());
+		LoadSoleStruct(&Out.GetUniqueID(), In);
 	}
 
 	static void Save(FMemberBuilder& Out, const FLazyObjectPtr& In, const FSaveContext& Ctx)
@@ -3001,8 +2989,7 @@ struct FMemoryPropertyBatch
 	static void Load(FLazyObjectPtr& Out, FStructLoadView In)
 	{
 		Out.ResetWeakPtr();
-		FMemberLoader Members(In);
-		LoadStruct(&Out.GetUniqueID(), Members.GrabStruct());
+		LoadSoleStruct(&Out.GetUniqueID(), In);
 	}
 
 #if UE_FNAME_OUTLINE_NUMBER
