PlainProps noting and writing UObject bindings

Bug fixes, improved error messages and moving some write functions

#rnx

diff --git a/Engine/Plugins/Experimental/PlainProps/Source/Private/PlainPropsBuild.cpp b/Engine/Plugins/Experimental/PlainProps/Source/Private/PlainPropsBuild.cpp
--- a/Engine/Plugins/Experimental/PlainProps/Source/Private/PlainPropsBuild.cpp
+++ b/Engine/Plugins/Experimental/PlainProps/Source/Private/PlainPropsBuild.cpp
@@ -120,9 +120,8 @@ FBuiltRange* CloneLeaves(FScratchAllocator& Scratch, uint64 Num, const void* InD
 
 void FMemberBuilder::BuildSuperStruct(FScratchAllocator& Scratch, const FStructDeclaration& Super, const FDebugIds& Debug)
 {
-	// If we need to support EMemberPresence::RequireAll for sub structs,
-	// we need access to struct declaration here or create an empty super
-	// structs that we throw away in BuildAndReset.
+	// If we need to support dense substructs, we need access to struct declaration here 
+	// or create an empty super structs that we throw away in BuildAndReset.
 	if (Members.IsEmpty() || (Members.Num() == 1 && IsSuper(Members[0].Schema.Type)))
 	{
 		return;
@@ -135,9 +134,13 @@ void FMemberBuilder::BuildSuperStruct(FScratchAllocator& Scratch, const FStructD
 
 FBuiltStruct* FMemberBuilder::BuildAndReset(FScratchAllocator& Scratch, const FStructDeclaration& Declared, const FDebugIds& Debug)
 {
-	checkf(!(Declared.Super && Declared.Occupancy == EMemberPresence::RequireAll),
-		TEXT("Requiring sub structs to be dense isn't implemented, %s"), *Debug.Print(Declared.Id));
 #if DO_CHECK
+	if (Declared.Occupancy == EMemberPresence::RequireAll)
+	{
+		checkf(Members.Num() == Declared.NumMembers, TEXT("'%s' requires exactly %d members but %d were added"), *Debug.Print(Declared.Id), Declared.NumMembers, Members.Num());
+		checkf(!Declared.Super, TEXT("Bug, dense substructs should fail DeclareStruct() check"));
+	}
+	
 	// Verify members were added in declared order
 	if (int32 Num = Members.Num())
 	{
diff --git a/Engine/Plugins/Experimental/PlainProps/Source/Private/PlainPropsBuildSchema.cpp b/Engine/Plugins/Experimental/PlainProps/Source/Private/PlainPropsBuildSchema.cpp
--- a/Engine/Plugins/Experimental/PlainProps/Source/Private/PlainPropsBuildSchema.cpp
+++ b/Engine/Plugins/Experimental/PlainProps/Source/Private/PlainPropsBuildSchema.cpp
@@ -215,7 +215,8 @@ static void* NoteStructOrEnum(FSchemasBuilder& AllSchemas, bool bStruct, FSchema
 
 void FStructSchemaBuilder::NoteMembersRecursively(const FBuiltStruct& Struct)
 {
-	check(Declaration.Occupancy != EMemberPresence::RequireAll || Struct.NumMembers == Declaration.NumMembers);
+	checkf(Declaration.Occupancy != EMemberPresence::RequireAll || Struct.NumMembers == Declaration.NumMembers,
+		TEXT("'%s' with %d members noted while declared to always have all %d members"), *AllSchemas.GetDebug().Print(Declaration.Id), Struct.NumMembers, Declaration.NumMembers);
 	MinMembers = FMath::Min(MinMembers, Struct.NumMembers);
 	
 	if (Struct.NumMembers == 0)
diff --git a/Engine/Plugins/Experimental/PlainProps/Source/Private/PlainPropsDeclare.cpp b/Engine/Plugins/Experimental/PlainProps/Source/Private/PlainPropsDeclare.cpp
--- a/Engine/Plugins/Experimental/PlainProps/Source/Private/PlainPropsDeclare.cpp
+++ b/Engine/Plugins/Experimental/PlainProps/Source/Private/PlainPropsDeclare.cpp
@@ -129,6 +129,7 @@ static const FStructDeclaration& Declare(TArray<TUniquePtr<FStructDeclaration>>&
 
 const FStructDeclaration& FDeclarations::DeclareStruct(FStructSchemaId DeclId, FTypeId Type, TConstArrayView<FMemberId> MemberOrder, EMemberPresence Occupancy, FOptionalStructSchemaId Super)
 {
+	checkf(!(Super && Occupancy == EMemberPresence::RequireAll), TEXT("'%s' is a dense substruct, this isn't supported, see BuildSuperStruct()"), *Debug.Print(DeclId));
 	return Declare(/* out */ DeclaredStructs, DeclId, Type, MemberOrder, Occupancy, Super, IntCastChecked<uint16>(MemberOrder.Num()));
 }
 
diff --git a/Engine/Plugins/Experimental/PlainProps/Source/Private/PlainPropsInternalBuild.h b/Engine/Plugins/Experimental/PlainProps/Source/Private/PlainPropsInternalBuild.h
--- a/Engine/Plugins/Experimental/PlainProps/Source/Private/PlainPropsInternalBuild.h
+++ b/Engine/Plugins/Experimental/PlainProps/Source/Private/PlainPropsInternalBuild.h
@@ -30,45 +30,5 @@ struct FBuiltRange
 	TConstArrayView64<const FBuiltStruct*>				AsStructs() const	{ return { reinterpret_cast<FBuiltStruct const* const*>(Data),	static_cast<int64>(Num) }; }
 };
 
-//////////////////////////////////////////////////////////////////////////
-
-inline void WriteData(TArray64<uint8>& Out, const void* Data, int64 Size)
-{
-	Out.Append(static_cast<const uint8*>(Data), Size);
-}
-
-template<typename ArrayType>
-void WriteArray(TArray64<uint8>& Out, const ArrayType& In)
-{
-	static_assert(TIsContiguousContainer<ArrayType>::Value);
-	WriteData(Out, In.GetData(), sizeof(typename ArrayType::ElementType) * In.Num());
-}
-
-template<typename T>
-void WriteAlignmentPadding(TArray64<uint8>& Out)
-{
-	Out.AddZeroed(Align(Out.Num(), alignof(T)) - Out.Num());
-}	
-
-template<typename T>
-void WriteAlignedArray(TArray64<uint8>& Out, TArrayView<T> In)
-{
-	WriteAlignmentPadding<T>(Out);
-	WriteArray(Out, In);
-}
-
-template<typename T>
-inline void WriteInt(TArray64<uint8>& Out, T Number)
-{
-	static_assert(std::is_integral_v<T>);
-	static_assert(PLATFORM_LITTLE_ENDIAN);
-	WriteData(Out, &Number, sizeof(T));
-}
-
-inline void WriteU32(TArray64<uint8>& Out, uint32 Int) { WriteInt(Out, Int); }
-inline void WriteU64(TArray64<uint8>& Out, uint64 Int) { WriteInt(Out, Int); }
-
-PLAINPROPS_API uint64 WriteSkippableSlice(TArray64<uint8>& Out, TConstArrayView64<uint8> Slice);
-
 
 } // namespace PlainProps
\ No newline at end of file
diff --git a/Engine/Plugins/Experimental/PlainProps/Source/Private/PlainPropsSave.cpp b/Engine/Plugins/Experimental/PlainProps/Source/Private/PlainPropsSave.cpp
--- a/Engine/Plugins/Experimental/PlainProps/Source/Private/PlainPropsSave.cpp
+++ b/Engine/Plugins/Experimental/PlainProps/Source/Private/PlainPropsSave.cpp
@@ -253,7 +253,7 @@ static ELeafWidth GetBasicWidth(FLeafBindType Leaf)
 
 [[nodiscard]] static FBuiltStruct* SaveRangeItem(const uint8* Struct, FDefaultStruct Default, const FSaveContext& Ctx)
 {
-	return SaveStructDeltaIfDiff(Struct, Default.Struct, Default.Id, Ctx);
+	return SaveStructDelta(Struct, Default.Struct, Default.Id, Ctx);
 }
 
 [[nodiscard]] static FMemberType ToMemberType(FMemberBindType In)
@@ -308,6 +308,20 @@ static void SaveMember(FMemberBuilder& Out, const void* Struct, FMemberId Name,
 	Out.AddStruct(Name, Member.Id, SaveStruct(At(Struct, Member.Offset), Member.Id, Ctx));
 }
 
+inline void SaveAllMembers(FMemberBuilder& Out, const void* Struct, const FSchemaBinding& Schema, TConstArrayView<FMemberId> MemberOrder, const FSaveContext& Ctx)
+{
+	for (FMemberVisitor It(Schema); It.HasMore(); )
+	{
+		FMemberId Name = MemberOrder[It.GetIndex()];
+		switch (It.PeekKind())
+		{
+			case EMemberKind::Leaf:		SaveMember(Out, Struct, Name, Ctx, It.GrabLeaf());		break;
+			case EMemberKind::Range:	SaveMember(Out, Struct, Name, Ctx, It.GrabRange());		break;
+			case EMemberKind::Struct:	SaveMember(Out, Struct, Name, Ctx, It.GrabStruct());	break;
+		}
+	}
+}
+
 FBuiltStruct* SaveStruct(const void* Struct, FStructSchemaId BindId, const FSaveContext& Ctx)
 {
 	const FStructDeclaration* Declaration = nullptr;
@@ -320,18 +334,7 @@ FBuiltStruct* SaveStruct(const void* Struct, FStructSchemaId BindId, const FSave
 	{
 		const FSchemaBinding& Schema = Ctx.Schemas.GetStruct(BindId);
 		Declaration = &Ctx.Declarations.Get(Schema.DeclId);	
-		TConstArrayView<FMemberId> MemberOrder = Declaration->GetMemberOrder();
-
-		for (FMemberVisitor It(Schema); It.HasMore(); )
-		{
-			FMemberId Name = MemberOrder[It.GetIndex()];
-			switch (It.PeekKind())
-			{
-				case EMemberKind::Leaf:		SaveMember(Out, Struct, Name, Ctx, It.GrabLeaf());		break;
-				case EMemberKind::Range:	SaveMember(Out, Struct, Name, Ctx, It.GrabRange());		break;
-				case EMemberKind::Struct:	SaveMember(Out, Struct, Name, Ctx, It.GrabStruct());	break;
-			}
-		}
+		SaveAllMembers(Out, Struct, Schema, Declaration->GetMemberOrder(), Ctx);
 	}
 
 	return Out.BuildAndReset(Ctx.Scratch, *Declaration, Ctx.Declarations.GetDebug());
@@ -526,7 +529,7 @@ static void SaveMemberDelta(FMemberBuilder& Out, const void* Struct, const void*
 {
 	if (FBuiltStruct* Delta = SaveStructDeltaIfDiff(At(Struct, Member.Offset), At(Default, Member.Offset), Member.Id, Ctx))
 	{
-		Out.AddStruct(Name, Member.Id, MoveTemp(Delta));
+		Out.AddStruct(Name, Member.Id, Delta);
 	}
 }
 
@@ -535,16 +538,23 @@ static const FStructDeclaration& SaveSchemaBoundStructDelta(FMemberBuilder& Out,
 	const FStructDeclaration& Declaration = Ctx.Declarations.Get(Schema.DeclId);
 	TConstArrayView<FMemberId> MemberOrder = Declaration.GetMemberOrder();
 
-	for (FMemberVisitor It(Schema); It.HasMore(); )
+	if (Declaration.Occupancy == EMemberPresence::AllowSparse)
 	{
-		FMemberId Name = MemberOrder[It.GetIndex()];
-		switch (It.PeekKind())
+		for (FMemberVisitor It(Schema); It.HasMore(); )
 		{
-			case EMemberKind::Leaf:		SaveMemberDelta(Out, Struct, Default, Name, Ctx, It.GrabLeaf());	break;
-			case EMemberKind::Range:	SaveMemberDelta(Out, Struct, Default, Name, Ctx, It.GrabRange());	break;
-			case EMemberKind::Struct:	SaveMemberDelta(Out, Struct, Default, Name, Ctx, It.GrabStruct());	break;
+			FMemberId Name = MemberOrder[It.GetIndex()];
+			switch (It.PeekKind())
+			{
+				case EMemberKind::Leaf:		SaveMemberDelta(Out, Struct, Default, Name, Ctx, It.GrabLeaf());	break;
+				case EMemberKind::Range:	SaveMemberDelta(Out, Struct, Default, Name, Ctx, It.GrabRange());	break;
+				case EMemberKind::Struct:	SaveMemberDelta(Out, Struct, Default, Name, Ctx, It.GrabStruct());	break;
+			}
 		}
 	}
+	else
+	{
+		SaveAllMembers(Out, Struct, Schema, MemberOrder, Ctx);
+	}
 
 	return Declaration;
 }
diff --git a/Engine/Plugins/Experimental/PlainProps/Source/Private/PlainPropsWrite.cpp b/Engine/Plugins/Experimental/PlainProps/Source/Private/PlainPropsWrite.cpp
--- a/Engine/Plugins/Experimental/PlainProps/Source/Private/PlainPropsWrite.cpp
+++ b/Engine/Plugins/Experimental/PlainProps/Source/Private/PlainPropsWrite.cpp
@@ -37,6 +37,8 @@ struct FWriteIds
 	TArray<FTypeId>						KeptParameters;
 	TArray<FNameId>						KeptNames;
 
+	const FDebugIds&					Debug;
+
 	FNameId								Remap(FNameId Old) const				{ return Names[Old.Idx].Get(); }
 	FMemberId							Remap(FMemberId Old) const				{ return { Remap(Old.Id) }; }
 	FFlatScopeId						Remap(FFlatScopeId Old) const			{ return { Remap(Old.Name) }; }
@@ -201,6 +203,7 @@ static void CopyUsedIds(TArray<T>& Out, const TBitArray<>& Used, C&& Ids)
 
 FWriteIds::FWriteIds(const FIdIndexerBase& Ids, const IStructBindIds& InBindIds, const FBuiltSchemas& Schemas, ESchemaFormat Format)
 : BindIds(InBindIds)
+, Debug(Ids)
 {
 	Structs.Init(NoId, Ids.NumStructs());
 	Enums.Init(NoId, Ids.NumEnums());
@@ -377,7 +380,9 @@ static void WriteSchema(TArray64<uint8>& Out, const FBuiltStructSchema& Struct,
 	TArray<FMemberId> MemberNames = NewIds.HasStableNames() ? RemapIds(NewIds, Struct.MemberNames) : Struct.MemberNames;
 	TArray<FSchemaId> InnerSchemas = GetInnerSchemas(Struct, NewIds, Inheritance);
 
-	check(MemberNames.Num() + UsesSuper(Inheritance) == MemberTypes.Num());
+	checkf(MemberNames.Num() + UsesSuper(Inheritance) == MemberTypes.Num(), 
+		TEXT("'%s' has %d member names and %d unnamed super but %d types"), 
+		*NewIds.Debug.Print(Struct.Type), MemberNames.Num(), UsesSuper(Inheritance), MemberTypes.Num());
 
 	// Zero init for determinism
 	alignas(FStructSchema) uint8 HeaderData[offsetof(FStructSchema, Footer)] = {};
@@ -453,9 +458,9 @@ static void WriteEnumConstants(TArray64<uint8>& Out, ELeafWidth Width, TConstArr
 static void WriteSchema(TArray64<uint8>& Out, const FBuiltEnumSchema& Enum, const FWriteIds& NewIds)
 {
 	bool bIsSequence = Enum.Mode == EEnumMode::Flag ? IsFlagSequence(Enum.Constants) : IsFlatSequence(Enum.Constants);
-	TArray<FNameId> Names = RemapIds(NewIds, Enum.Names);
+	TArray<FNameId> Names = NewIds.HasStableNames() ? RemapIds(NewIds, Enum.Names) : Enum.Names;
 
-	FEnumSchema BinaryHeader = { NewIds.Remap(Enum.Type) };
+	FEnumSchema BinaryHeader = { NewIds.HasStableNames() ? NewIds.Remap(Enum.Type) : Enum.Type };
 	BinaryHeader.FlagMode = Enum.Mode == EEnumMode::Flag;
 	BinaryHeader.ExplicitConstants = !bIsSequence;
 	BinaryHeader.Width = Enum.Width;
diff --git a/Engine/Plugins/Experimental/PlainProps/Source/Public/PlainPropsBind.h b/Engine/Plugins/Experimental/PlainProps/Source/Public/PlainPropsBind.h
--- a/Engine/Plugins/Experimental/PlainProps/Source/Public/PlainPropsBind.h
+++ b/Engine/Plugins/Experimental/PlainProps/Source/Public/PlainPropsBind.h
@@ -602,7 +602,7 @@ struct IDefaultStructs
 
 ////////////////////////////////////////////////////////////////////////////////////////////////
 
-struct FStructBindIds : IStructBindIds
+struct FStructBindIds final : IStructBindIds
 {
 	FStructBindIds(const FCustomBindings& InCustoms, const FSchemaBindings& InSchemas)
 	: Customs(InCustoms)
@@ -612,7 +612,7 @@ struct FStructBindIds : IStructBindIds
 	const FCustomBindings& Customs;
 	const FSchemaBindings& Schemas;
 
-	virtual FStructSchemaId GetDeclId(FStructSchemaId BindId) const override final;
+	PLAINPROPS_API virtual FStructSchemaId GetDeclId(FStructSchemaId BindId) const override;
 };
 
 ////////////////////////////////////////////////////////////////////////////////////////////////
diff --git a/Engine/Plugins/Experimental/PlainProps/Source/Public/PlainPropsWrite.h b/Engine/Plugins/Experimental/PlainProps/Source/Public/PlainPropsWrite.h
--- a/Engine/Plugins/Experimental/PlainProps/Source/Public/PlainPropsWrite.h
+++ b/Engine/Plugins/Experimental/PlainProps/Source/Public/PlainPropsWrite.h
@@ -5,6 +5,7 @@
 #include "PlainPropsTypes.h"
 #include "Containers/Array.h"
 #include "Templates/UniquePtr.h"
+#include <type_traits>
 
 namespace PlainProps
 {
@@ -37,4 +38,39 @@ private:
 	TUniquePtr<FWriteIds>							NewIds;
 };
 
+//////////////////////////////////////////////////////////////////////////
+
+inline void WriteData(TArray64<uint8>& Out, const void* Data, int64 Size)
+{
+	Out.Append(static_cast<const uint8*>(Data), Size);
+}
+
+template<typename ArrayType>
+void WriteArray(TArray64<uint8>& Out, const ArrayType& In) requires (!!TIsContiguousContainer<ArrayType>::Value)
+{
+	WriteData(Out, In.GetData(), sizeof(typename ArrayType::ElementType) * In.Num());
+}
+
+template<typename T>
+void WriteAlignmentPadding(TArray64<uint8>& Out)
+{
+	Out.AddZeroed(Align(Out.Num(), alignof(T)) - Out.Num());
+}	
+
+template<typename T>
+void WriteAlignedArray(TArray64<uint8>& Out, TArrayView<T> In)
+{
+	WriteAlignmentPadding<T>(Out);
+	WriteArray(Out, In);
+}
+
+template<typename T>
+inline void WriteInt(TArray64<uint8>& Out, T Number) requires (std::is_integral_v<T> && !!PLATFORM_LITTLE_ENDIAN)
+{
+	WriteData(Out, &Number, sizeof(T));
+}
+
+PLAINPROPS_API uint64 WriteSkippableSlice(TArray64<uint8>& Out, TConstArrayView64<uint8> Slice);
+
+
 } // namespace PlainProps
diff --git a/Engine/Plugins/Experimental/PlainProps/Source/Tests/PlainPropsTests.cpp b/Engine/Plugins/Experimental/PlainProps/Source/Tests/PlainPropsTests.cpp
--- a/Engine/Plugins/Experimental/PlainProps/Source/Tests/PlainPropsTests.cpp
+++ b/Engine/Plugins/Experimental/PlainProps/Source/Tests/PlainPropsTests.cpp
@@ -229,24 +229,24 @@ TArray64<uint8> FTestBatchBuilder::Write()
 		const FAnsiString& Str = ResolveName(Name);
 		WriteData(Tmp, *Str, Str.Len() + 1);
 	}
-	WriteU32(Out, TestMagics[0]);
+	WriteInt(Out, TestMagics[0]);
 	WriteSkippableSlice(Out, Tmp);
 	Tmp.Reset();
 
 	// Write schemas
-	WriteU32(Out, TestMagics[1]);
+	WriteInt(Out, TestMagics[1]);
 	Writer.WriteSchemas(/* Out */ Tmp);
 	WriteAlignmentPadding<uint32>(Out);
-	WriteU32(Out, IntCastChecked<uint32>(Tmp.Num()));
+	WriteInt(Out, IntCastChecked<uint32>(Tmp.Num()));
 	WriteArray(Out, Tmp);
 	Tmp.Reset();
 
 	// Write objects
-	WriteU32(Out, TestMagics[2]);
+	WriteInt(Out, TestMagics[2]);
 	for (const TPair<FStructSchemaId, FBuiltStruct*>& Object : Objects)
 	{
-		WriteU32(/* out */ Tmp, TestMagics[3]);
-		WriteU32(/* out */ Tmp, Writer.GetWriteId(Object.Key).Get().Idx);
+		WriteInt(/* out */ Tmp, TestMagics[3]);
+		WriteInt(/* out */ Tmp, Writer.GetWriteId(Object.Key).Get().Idx);
 		Writer.WriteMembers(/* out */ Tmp, Object.Key, *Object.Value);
 		WriteSkippableSlice(Out, Tmp);
 		Tmp.Reset();
@@ -254,7 +254,7 @@ TArray64<uint8> FTestBatchBuilder::Write()
 
 	// Write object terminator
 	WriteSkippableSlice(Out, TConstArrayView64<uint8>());
-	WriteU32(Out, TestMagics[4]);
+	WriteInt(Out, TestMagics[4]);
 		
 	return Out;
 }
diff --git a/Engine/Plugins/Experimental/PlainProps/Source/Tests/PlainPropsUeCoreTests.cpp b/Engine/Plugins/Experimental/PlainProps/Source/Tests/PlainPropsUeCoreTests.cpp
--- a/Engine/Plugins/Experimental/PlainProps/Source/Tests/PlainPropsUeCoreTests.cpp
+++ b/Engine/Plugins/Experimental/PlainProps/Source/Tests/PlainPropsUeCoreTests.cpp
@@ -151,7 +151,7 @@ public:
 	TArray64<uint8>				Write() const;
 
 private:
-	using IdBuiltStructPair = TPair<FStructSchemaId, FBuiltStruct*>;
+	using IdBuiltStructPair = TPair<FStructSchemaId, const FBuiltStruct*>;
 	TArray<IdBuiltStructPair>	SavedObjects;
 	FNameBinding				SavedNames;
 	FCustomBindingsOverlay		Customs;
@@ -175,7 +175,7 @@ template<class T>
 bool FBatchSaver::SaveDelta(const T& Object, const T& Default) 
 {
 	FDualStructSchemaId Id = IndexStructDualId<FIds, TTypename<std::remove_reference_t<T>>>();
-	if (FBuiltStruct* Delta = SaveStructDeltaIfDiff(&Object, &Default, Id.BindId, {GTypes, GSchemas, Customs, Scratch}))
+	if (const FBuiltStruct* Delta = SaveStructDeltaIfDiff(&Object, &Default, Id.BindId, {GTypes, GSchemas, Customs, Scratch}))
 	{
 		SavedObjects.Emplace(Id.DeclId, MoveTemp(Delta));
 		return true;
@@ -186,7 +186,7 @@ bool FBatchSaver::SaveDelta(const T& Object, const T& Default)
 template<typename ArrayType>
 void WriteNumAndArray(TArray64<uint8>& Out, const ArrayType& Items)
 {
-	WriteU32(Out, IntCastChecked<uint32>(Items.Num()));
+	WriteInt(Out, IntCastChecked<uint32>(Items.Num()));
 	WriteArray(Out, Items);
 }
 
@@ -221,11 +221,11 @@ TArray64<uint8> FBatchSaver::Write() const
 	
 	// Write ids. Just copying in-memory FNames, a stable format might use SaveNameBatch().
 	TArray64<uint8> Out;
-	WriteU32(Out, Magics[0]);
+	WriteInt(Out, Magics[0]);
 	WriteNumAndArray(Out, TArrayView<const FName, int32>(UsedNames));
 
 	// Write schemas
-	WriteU32(Out, Magics[1]);
+	WriteInt(Out, Magics[1]);
 	WriteAlignmentPadding<uint32>(Out);
 	TArray64<uint8> Tmp;
 	Writer.WriteSchemas(/* Out */ Tmp);
@@ -233,11 +233,11 @@ TArray64<uint8> FBatchSaver::Write() const
 	Tmp.Reset();
 
 	// Write objects
-	WriteU32(Out, Magics[2]);
-	for (const TPair<FStructSchemaId, FBuiltStruct*>& Object : SavedObjects)
+	WriteInt(Out, Magics[2]);
+	for (const TPair<FStructSchemaId, const FBuiltStruct*>& Object : SavedObjects)
 	{
-		WriteU32(/* out */ Tmp, Magics[3]);
-		WriteU32(/* out */ Tmp, Writer.GetWriteId(Object.Key).Get().Idx);
+		WriteInt(/* out */ Tmp, Magics[3]);
+		WriteInt(/* out */ Tmp, Writer.GetWriteId(Object.Key).Get().Idx);
 		Writer.WriteMembers(/* out */ Tmp, Object.Key, *Object.Value);
 		WriteSkippableSlice(Out, Tmp);
 		Tmp.Reset();
@@ -245,11 +245,11 @@ TArray64<uint8> FBatchSaver::Write() const
 
 	// Write object terminator
 	WriteSkippableSlice(Out, TConstArrayView64<uint8>());
-	WriteU32(Out, Magics[4]);
+	WriteInt(Out, Magics[4]);
 	
 	// Write names
 	WriteNumAndArray(Out, SavedNames.Names.Array());
-	WriteU32(Out, Magics[5]);
+	WriteInt(Out, Magics[5]);
 		
 	return Out;
 }
diff --git a/Engine/Plugins/Experimental/PlainPropsUObject/Source/Private/PlainPropsUObjectBindings.cpp b/Engine/Plugins/Experimental/PlainPropsUObject/Source/Private/PlainPropsUObjectBindings.cpp
--- a/Engine/Plugins/Experimental/PlainPropsUObject/Source/Private/PlainPropsUObjectBindings.cpp
+++ b/Engine/Plugins/Experimental/PlainPropsUObject/Source/Private/PlainPropsUObjectBindings.cpp
@@ -239,7 +239,8 @@ FCommonStructIds::FCommonStructIds(const FCommonScopeIds& Scopes, TIdIndexer<FNa
 {}
 
 FCommonMemberIds::FCommonMemberIds(TIdIndexer<FName>& Names)
-: Value(					Names.NameMember("Value"))
+: Key(						Names.NameMember("Key"))
+, Value(					Names.NameMember("Value"))
 , Add(						Names.NameMember("Add"))
 , Del(						Names.NameMember("Del"))
 , Id(						Names.NameMember("Id"))
@@ -376,18 +377,45 @@ static EMemberPresence GetOccupancy(const UStruct* Struct)
 	return EMemberPresence::AllowSparse;
 }
 
+using FMemberArray = TArray<FMemberId, TInlineAllocator<64>>;
+
+inline void DeclareMembers(FMemberArray& Out, const UStruct* Struct)
+{
+	for (FProperty* It = Struct->PropertyLink; It && It->GetOwner<UStruct>() == Struct; It = It->PropertyLinkNext)
+	{
+		Out.Add(GUE.Names.NameMember(It->GetFName()));
+	}
+}
+
+// Must match BindSuperMembers
+static void DeclareSuperMembers(FMemberArray& Out, const UStruct* Struct)
+{
+	if (const UStruct* Super = Struct->GetInheritanceSuper())
+	{
+		DeclareSuperMembers(Out, Super);
+		if (ShouldBind(Super))
+		{
+			DeclareMembers(Out, Super);
+		}
+	}
+}
+
 static FStructSchemaId DeclareStruct(const UStruct* Struct)
 {
 	FTypeId Type = IndexType(Struct);
 	FStructSchemaId Id = GUE.Names.IndexStruct(Type);
 	FOptionalStructSchemaId  Super = IndexSuper(Struct);
 	EMemberPresence Occupancy = GetOccupancy(Struct);
-	TArray<FMemberId, TInlineAllocator<64>> Members;
-	for (FProperty* It = Struct->PropertyLink; It && It->GetOwner<UStruct>() == Struct; It = It->PropertyLinkNext)
+	
+	FMemberArray Members;
+	if (Super && Occupancy == EMemberPresence::RequireAll)
 	{
-		Members.Add(GUE.Names.NameMember(It->GetFName()));
+		// Flatten inheritance chain for dense structs
+		Super = NoId;
+		DeclareSuperMembers(/* out */ Members, Struct);
 	}
-	
+	DeclareMembers(/* out */ Members, Struct);
+
 	GUE.Types.DeclareStruct(Id, Type, Members, Occupancy, Super);
 
 	return Id;
@@ -2559,6 +2587,7 @@ public:
 
 	void								BindMember(FProperty* Property) { Members.Emplace(BindProperty(Property)); }
 	TConstArrayView<FMemberBinding>		GetMembers() const { return Members; }
+	bool								IsDense() { return GetOwnerOccupancy() == EMemberPresence::RequireAll; }
 
 private:
 	const FStructSchemaId							Owner;
@@ -2843,7 +2872,31 @@ private:
 	}
 };
 
-void BindStruct(FStructSchemaId Id, UStruct* Struct)
+inline void BindMembers(FPropertyBinder& Out, const UStruct* Struct)
+{
+	for (FProperty* It = Struct->PropertyLink; It && It->GetOwner<UStruct>() == Struct; It = It->PropertyLinkNext)
+	{
+		if (ShouldBind(It))
+		{
+			Out.BindMember(It);
+		}
+	}
+}
+
+// Must match DeclareSuperMembers
+void BindSuperMembers(FPropertyBinder& Out, const UStruct* Struct)
+{
+	if (const UStruct* Super = Struct->GetInheritanceSuper())
+	{
+		BindSuperMembers(Out, Super);
+		if (ShouldBind(Super))
+		{
+			BindMembers(Out, Super);
+		}
+	}
+}
+
+void BindStruct(FStructSchemaId Id, const UStruct* Struct)
 {
 	if (GUE.Customs.FindStruct(Id))
 	{
@@ -2851,13 +2904,12 @@ void BindStruct(FStructSchemaId Id, UStruct* Struct)
 	}
 
 	FPropertyBinder Binder(Id);
-	for (FProperty* It = Struct->PropertyLink; It && It->GetOwner<UStruct>() == Struct; It = It->PropertyLinkNext)
+	if (Struct->GetSuperStruct() && Binder.IsDense())
 	{
-		if (ShouldBind(It))
-		{
-			Binder.BindMember(It);
-		}
+		BindSuperMembers(/* out */ Binder, Struct);
 	}
+	BindMembers(/* out */ Binder, Struct);
+
 	GUE.Schemas.BindStruct(Id, Id, Binder.GetMembers());
 	
 	// Don't bind CDOs, object defaults are passed in from top and objects aren't owned by containers 
@@ -3103,65 +3155,61 @@ public:
 		SavedObjects.Emplace(Id, bFlat ? SaveStruct(Object, Id, FlatCtx) : SaveStructDelta(Object, CDO, Id, DeltaCtx));
 	}
 
-	//TArray64<uint8>	Write() const
-	//{
-	//	// Build partial schemas
-	//	const FStructBindIds BindIds(GUE.DeltaCustoms, GUE.Schemas);
-	//	FSchemasBuilder SchemaBuilders(GUE.Types, BindIds, Scratch);
-	//	for (const IdBuiltStructPair& Object : SavedObjects)
-	//	{
-	//		SchemaBuilders.NoteStructAndMembers(Object.Key, *Object.Value);
-	//	}
-	//	FBuiltSchemas Schemas = SchemaBuilders.Build(); 
-
-	//	// Filter out declared but unused names and ids
-	//	FWriter Writer(GUE.Names, BindIds, Schemas, ESchemaFormat::InMemoryNames);
-	//	TArray<FName> UsedNames;
-	//	UsedNames.Reserve(Writer.GetUsedNames().Num());
-	//	for (FNameId Name : Writer.GetUsedNames())
-	//	{
-	//		UsedNames.Add(GUE.Names.ResolveName(Name));
-	//	}
-	//
-	//	// Write ids. Just copying in-memory FNames, a stable format might use SaveNameBatch().
-	//	TArray64<uint8> Out;
-	//	WriteU32(Out, Magics[0]);
-	//	WriteNumAndArray(Out, TArrayView<const FName, int32>(UsedNames));
-
-	//	// Write schemas
-	//	WriteU32(Out, Magics[1]);
-	//	WriteAlignmentPadding<uint32>(Out);
-	//	TArray64<uint8> Tmp;
-	//	Writer.WriteSchemas(/* Out */ Tmp);
-	//	WriteNumAndArray(Out, TArrayView<const uint8, int64>(Tmp));
-	//	Tmp.Reset();
-
-	//	// Write objects
-	//	WriteU32(Out, Magics[2]);
-	//	for (const TPair<FStructSchemaId, FBuiltStruct*>& Object : SavedObjects)
-	//	{
-	//		WriteU32(/* out */ Tmp, Magics[3]);
-	//		WriteU32(/* out */ Tmp, Writer.GetWriteId(Object.Key).Get().Idx);
-	//		Writer.WriteMembers(/* out */ Tmp, Object.Key, *Object.Value);
-	//		WriteSkippableSlice(Out, Tmp);
-	//		Tmp.Reset();
-	//	}
-
-	//	// Write object terminator
-	//	WriteSkippableSlice(Out, TConstArrayView64<uint8>());
-	//	WriteU32(Out, Magics[4]);
-	//
-	//	// Write names
-	//	WriteNumAndArray(Out, SavedNames.Names.Array());
-	//	WriteU32(Out, Magics[5]);
-	//	
-	//	return Out;
-	//}
+	TArray64<uint8>	Write() const
+	{
+		// Build partial schemas
+		const FStructBindIds BindIds(GUE.DeltaCustoms, GUE.Schemas);
+		FSchemasBuilder SchemaBuilders(GUE.Types, BindIds, Scratch);
+		for (const IdBuiltStructPair& Object : SavedObjects)
+		{
+			SchemaBuilders.NoteStructAndMembers(Object.Key, *Object.Value);
+		}
+		FBuiltSchemas Schemas = SchemaBuilders.Build(); 
+
+		// Filter out declared but unused names and ids
+		FWriter Writer(GUE.Names, BindIds, Schemas, ESchemaFormat::InMemoryNames);
+		TArray<FName> UsedNames;
+		UsedNames.Reserve(Writer.GetUsedNames().Num());
+		for (FNameId Name : Writer.GetUsedNames())
+		{
+			UsedNames.Add(GUE.Names.ResolveName(Name));
+		}
+	
+		// Write ids. Just copying in-memory FNames, a stable format might use SaveNameBatch().
+		TArray64<uint8> Out;
+		WriteInt(Out, Magics[0]);
+		WriteNumAndArray(Out, TArrayView<const FName, int32>(UsedNames));
+
+		// Write schemas
+		WriteInt(Out, Magics[1]);
+		WriteAlignmentPadding<uint32>(Out);
+		TArray64<uint8> Tmp;
+		Writer.WriteSchemas(/* Out */ Tmp);
+		WriteNumAndArray(Out, TArrayView<const uint8, int64>(Tmp));
+		Tmp.Reset();
+
+		// Write objects
+		WriteInt(Out, Magics[2]);
+		for (const TPair<FStructSchemaId, FBuiltStruct*>& Object : SavedObjects)
+		{
+			WriteInt(/* out */ Tmp, Magics[3]);
+			WriteInt(/* out */ Tmp, Writer.GetWriteId(Object.Key).Get().Idx);
+			Writer.WriteMembers(/* out */ Tmp, Object.Key, *Object.Value);
+			WriteSkippableSlice(Out, Tmp);
+			Tmp.Reset();
+		}
+
+		// Write object terminator
+		WriteSkippableSlice(Out, TConstArrayView64<uint8>());
+		WriteInt(Out, Magics[4]);
+		
+		return Out;
+	}
 
 private:
 	using IdBuiltStructPair = TPair<FStructSchemaId, FBuiltStruct*>;
 	using FPropertyBindings = TCustomPropertyBindings<PropertyBatchType>;
-	static constexpr uint32 Magics[] = { 0xFEEDF00D, 0xABCD1234, 0xDADADAAA, 0x99887766, 0xF0F1F2F3 , 0x00112233};
+	static constexpr uint32 Magics[] = { 0xFEEDF00D, 0xABCD1234, 0xDADADAAA, 0x99887766, 0xF0F1F2F3 };
 
 	TArray<IdBuiltStructPair>	SavedObjects;
 	PropertyBatchType			SavedProperties;
@@ -3171,29 +3219,34 @@ private:
 	FSaveContext				FlatCtx;
 	FSaveContext				DeltaCtx;
 
-	//template<typename ArrayType>
-	//static void WriteNumAndArray(TArray64<uint8>& Out, const ArrayType& Items)
-	//{
-	//	WriteU32(Out, IntCastChecked<uint32>(Items.Num()));
-	//	WriteArray(Out, Items);
-	//}
-
-	//template<typename T>
-	//static TConstArrayView<T> GrabNumAndArray(/* in-out */ FByteReader& It)
-	//{
-	//	uint32 Num = It.Grab<uint32>();
-	//	return MakeArrayView(reinterpret_cast<const T*>(It.GrabBytes(Num * sizeof(T))), Num);
-	//}
+	template<typename ArrayType>
+	static void WriteNumAndArray(TArray64<uint8>& Out, const ArrayType& Items)
+	{
+		WriteInt(Out, IntCastChecked<uint32>(Items.Num()));
+		WriteArray(Out, Items);
+	}
+
+	template<typename T>
+	static TConstArrayView<T> GrabNumAndArray(/* in-out */ FByteReader& It)
+	{
+		uint32 Num = It.Grab<uint32>();
+		return MakeArrayView(reinterpret_cast<const T*>(It.GrabBytes(Num * sizeof(T))), Num);
+	}
 };
 
 
 static void Roundtrip()
 {
-	TBatchSaver<FMemoryPropertyBatch> Saver(GUObjectArray.GetObjectArrayNum());
-	for (TObjectIterator<UObject> It; It; ++It)
+	TArray64<uint8> Data;
 	{
-		Saver.Save(*It);
+		TBatchSaver<FMemoryPropertyBatch> Batch(GUObjectArray.GetObjectArrayNum());
+		for (TObjectIterator<UObject> It; It; ++It)
+		{
+			Batch.Save(*It);
+		}
+		Data = Batch.Write();
 	}
+
 }
 
 static void TestBindings()
