PlainProps add new diffing APIs and fix some loading bugs

#rnx

diff --git a/Engine/Plugins/Experimental/PlainProps/Source/Private/PlainPropsDiff.cpp b/Engine/Plugins/Experimental/PlainProps/Source/Private/PlainPropsDiff.cpp
new file mode 100644
--- /dev/null
+++ b/Engine/Plugins/Experimental/PlainProps/Source/Private/PlainPropsDiff.cpp
@@ -0,0 +1,127 @@
+// Copyright Epic Games, Inc. All Rights Reserved.
+
+#include "PlainPropsDiff.h"
+#include "PlainPropsBind.h"
+#include "PlainPropsInternalBind.h"
+#include "PlainPropsInternalDiff.h"
+
+namespace PlainProps
+{
+
+static bool DiffItem(const void* A, const void* B, FDiffContext& Ctx, FRangeMemberBinding Range);
+static bool DiffItem(const void* A, const void* B, FDiffContext& Ctx, FStructSchemaId Struct);
+
+static bool DiffMember(FLeafMemberBinding Member, const void* OwnerA, const void* OwnerB, FDiffContext& Ctx, FOptionalMemberId Name)
+{
+	const uint8* A = At(OwnerA, Member.Offset);
+	const uint8* B = At(OwnerB, Member.Offset);
+	if (DiffLeaf(A, B, Member.Leaf))
+	{
+		Ctx.Out.Add({Member.Leaf.Pack(), Name, {Member.Enum}, A, B});
+		return true;
+	}
+	return false;
+}
+
+static bool DiffMember(FRangeMemberBinding Member, const void* OwnerA, const void* OwnerB, FDiffContext& Ctx, FOptionalMemberId Name)
+{
+	FRangeBinding Binding = Member.RangeBindings[0];
+	FMemberBindType InnerType = Member.InnerTypes[0];
+	EMemberKind InnerKind = InnerType.GetKind();
+	const uint8* A = At(OwnerA, Member.Offset);
+	const uint8* B = At(OwnerB, Member.Offset);
+
+	bool bDiff;
+	if (Binding.IsLeafBinding())
+	{
+		bDiff = Binding.AsLeafBinding().DiffLeaves(A, B);
+	}
+	else
+	{
+		const IItemRangeBinding& ItemBinding = Binding.AsItemBinding();
+		if (InnerKind == EMemberKind::Struct)
+		{
+			bDiff = DiffItemRange(A, B, ItemBinding, Ctx, static_cast<FStructSchemaId>(Member.InnermostSchema.Get()));
+		}
+		else if (InnerKind == EMemberKind::Leaf)
+		{
+			bDiff = DiffItemRange(A, B, ItemBinding, Ctx, SizeOf(GetBasicWidth(InnerType.AsLeaf())));
+		}
+		else
+		{
+			bDiff = DiffItemRange(A, B, ItemBinding, Ctx, GetInnerRange(Member));
+		}
+	}
+	if (bDiff)
+	{
+		Ctx.Out.Add({FMemberBindType(Binding.GetSizeType()), Name, {.Range = Binding}, A, B});
+	}
+	return bDiff;
+}
+
+static bool DiffMembers(const FSchemaBinding& Schema, const void* A, const void* B, FDiffContext& Ctx);
+
+inline bool DiffStruct(FStructSchemaId Id, const void* A, const void* B, FDiffContext& Ctx)
+{
+	if (const ICustomBinding* Custom = Ctx.Customs.FindStruct(Id))
+	{
+		return Custom->DiffCustom(A, B);
+	}
+	return DiffMembers(Ctx.Schemas.GetStruct(Id), A, B, Ctx);
+}
+
+static bool DiffMember(FStructMemberBinding Member, const void* OwnerA, const void* OwnerB, FDiffContext& Ctx, FMemberId Name)
+{
+	const uint8* A = At(OwnerA, Member.Offset);
+	const uint8* B = At(OwnerB, Member.Offset);
+	if (DiffStruct(Member.Id, A, B, Ctx))
+	{
+		Ctx.Out.Add({FMemberBindType(Member.Type), Name, {.Struct = Member.Id}, A, B});
+		return true;
+	}
+	return false;
+}
+
+static bool DiffMembers(const FSchemaBinding& Schema, const void* A, const void* B, FDiffContext& Ctx)
+{
+	const FStructDeclaration& Declaration = Ctx.Types.Get(Schema.DeclId);
+	TConstArrayView<FMemberId> MemberOrder = Declaration.GetMemberOrder();
+
+	bool bDiff = false;
+	for (FMemberVisitor It(Schema); !bDiff && It.HasMore(); )
+	{
+		FMemberId Name = MemberOrder[It.GetIndex()];
+		switch (It.PeekKind())
+		{
+			case EMemberKind::Leaf:		bDiff = DiffMember(It.GrabLeaf(), A, B, Ctx, Name);	break;
+			case EMemberKind::Range:	bDiff = DiffMember(It.GrabRange(), A, B, Ctx, Name);	break;
+			case EMemberKind::Struct:	bDiff = DiffMember(It.GrabStruct(), A, B, Ctx, Name);	break;
+		}
+	}
+
+	return bDiff;
+}
+
+static bool DiffItem(const void* A, const void* B, FDiffContext& Ctx, FRangeMemberBinding Range)
+{
+	return DiffMember(Range, A, B, Ctx, NoId);
+}
+
+static bool DiffItem(const void* A, const void* B, FDiffContext& Ctx, FStructSchemaId Id)
+{
+	return DiffStruct(Id, A, B, Ctx);
+}
+
+////////////////////////////////////////////////////////////////////////////////////////////////
+
+bool DiffStructs(FStructSchemaId BindId, const void* A, const void* B, FDiffContext& Ctx)
+{
+	if (const ICustomBinding* Custom = Ctx.Customs.FindStruct(BindId))
+	{
+		// Currently no way for DiffCustom to emit deeper diff into
+		return Custom->DiffCustom(A, B);
+	}
+	return DiffMembers(Ctx.Schemas.GetStruct(BindId), A, B, Ctx);
+}
+
+} // namespace PlainProps
\ No newline at end of file
diff --git a/Engine/Plugins/Experimental/PlainProps/Source/Private/PlainPropsInternalDiff.h b/Engine/Plugins/Experimental/PlainProps/Source/Private/PlainPropsInternalDiff.h
new file mode 100644
--- /dev/null
+++ b/Engine/Plugins/Experimental/PlainProps/Source/Private/PlainPropsInternalDiff.h
@@ -0,0 +1,140 @@
+// Copyright Epic Games, Inc. All Rights Reserved.
+
+#pragma once
+
+#include "PlainPropsBind.h"
+#include "PlainPropsTypes.h"
+
+namespace PlainProps 
+{
+
+inline const uint8* At(const void* Ptr, SIZE_T Offset)
+{
+	return static_cast<const uint8*>(Ptr) + Offset;
+}
+
+inline int64 DiffLeaf(const uint8* A, const uint8* B, FUnpackedLeafBindType Leaf)
+{
+	if (Leaf.Type == ELeafBindType::BitfieldBool)
+	{
+		uint32 Mask = 1u << Leaf.BitfieldIdx;
+		return (*A & Mask) - (*B & Mask);
+	}
+
+	switch (Leaf.Width)
+	{
+	case ELeafWidth::B8:	return FMemory::Memcmp(A, B, 1);
+	case ELeafWidth::B16:	return FMemory::Memcmp(A, B, 2);
+	case ELeafWidth::B32:	return FMemory::Memcmp(A, B, 4);
+	case ELeafWidth::B64:	return FMemory::Memcmp(A, B, 8);
+	}
+
+	unimplemented();
+	return 0;
+}
+
+inline ELeafWidth GetBasicWidth(FLeafBindType Leaf)
+{
+	checkf(Leaf.Bind.Type != ELeafBindType::BitfieldBool, TEXT("Arrays of bitfields is not a thing"));
+	return Leaf.Basic.Width;
+}
+
+////////////////////////////////////////////////////////////////////////////////////////////////
+
+// Helps diff ranges
+struct FItemRangeReader
+{
+	explicit FItemRangeReader(const void* Range, const IItemRangeBinding& Binding)
+	: Ctx{{Range}}
+	{
+		ReadItems(Binding);
+	}
+
+	FSaveRangeContext Ctx;
+	const uint8* SliceIt = nullptr;
+	uint64 SliceNum = 0;
+
+	void ReadItems(const IItemRangeBinding& Binding)
+	{
+		Binding.ReadItems(Ctx);
+		SliceIt = static_cast<const uint8*>(Ctx.Items.Slice.Data);
+		SliceNum = Ctx.Items.Slice.Num;
+	}
+
+	void RefillItems(const IItemRangeBinding& Binding)
+	{
+		if (SliceNum == 0)
+		{
+			ReadItems(Binding);
+			check(SliceNum > 0);
+		}
+	}
+
+	const uint8* GrabItems(uint64 Num, uint32 Stride)
+	{
+		const uint8* Out = SliceIt;
+		SliceIt += Num * Stride; 
+		SliceNum -= Num;
+		return Out;
+	}
+};
+
+template<class ContextType>
+inline bool DiffItemSlice(const uint8* A, const uint8* B, uint64 Num, uint32 Stride, ContextType&, SIZE_T LeafSize)
+{
+	return !!FMemory::Memcmp(A, B, Num * Stride);
+}
+
+template<typename ItemType, class ContextType>
+bool DiffItemSlice(const uint8* A, const uint8* B, uint64 Num, uint32 Stride, ContextType& Ctx, ItemType&& Member)
+{
+	for (const uint8* EndA = A + Num * Stride; A != EndA; A += Stride, B += Stride)
+	{
+		if (DiffItem(A, B, Ctx, Member))
+		{
+			return true;
+		}
+	}
+
+	return false;
+}
+
+template<class T, class ContextType>
+bool DiffItemRange(const void* RangeA, const void* RangeB, const IItemRangeBinding& Binding, ContextType& OuterCtx, T&& ItemSchema)
+{
+	FItemRangeReader A(RangeA, Binding);
+	FItemRangeReader B(RangeB, Binding);
+	if (A.Ctx.Items.NumTotal != B.Ctx.Items.NumTotal)
+	{
+		return true;
+	}
+
+	if (const uint64 NumTotal = A.Ctx.Items.NumTotal)
+	{
+		check(A.Ctx.Items.Stride == B.Ctx.Items.Stride);
+		const uint32 Stride = A.Ctx.Items.Stride;
+		while (true)
+		{
+			uint64 Num = FMath::Min(A.SliceNum, B.SliceNum);
+			if (DiffItemSlice(A.GrabItems(Num, Stride), B.GrabItems(Num, Stride), Num, Stride, OuterCtx, ItemSchema))
+			{
+				return true;
+			}
+			else if (A.Ctx.Request.NumRead + Num >= NumTotal)
+			{
+				check(A.Ctx.Request.NumRead + Num == NumTotal);	
+				check(B.Ctx.Request.NumRead + Num == NumTotal);	
+				return false;
+			}
+			
+			A.Ctx.Request.NumRead += Num;
+			B.Ctx.Request.NumRead += Num;
+			A.RefillItems(Binding);
+			B.RefillItems(Binding);
+		}
+	}
+
+	return false;
+}
+
+} // namespace PlainProps
\ No newline at end of file
diff --git a/Engine/Plugins/Experimental/PlainProps/Source/Private/PlainPropsLoad.cpp b/Engine/Plugins/Experimental/PlainProps/Source/Private/PlainPropsLoad.cpp
--- a/Engine/Plugins/Experimental/PlainProps/Source/Private/PlainPropsLoad.cpp
+++ b/Engine/Plugins/Experimental/PlainProps/Source/Private/PlainPropsLoad.cpp
@@ -738,16 +738,17 @@ FLoadBatchPtr CreateLoadPlans(FReadBatchId ReadId, const FDeclarations& Declarat
 
 ////////////////////////////////////////////////////////////////////////////
 
-inline static void SetBit(uint8& Out, uint8 Idx, bool bValue)
+inline void SetBit(uint8& Out, uint8 Idx, bool bValue)
 {
-	uint8 Mask = IntCastChecked<uint8>(1 << Idx);
+	checkSlow(Idx < 8);
+	uint8 Mask(1 << Idx);
 	if (bValue)
 	{
 		Out |= Mask;
 	}
 	else
 	{
-		Out &= Mask;
+		Out &= ~Mask;
 	}	
 }
 
diff --git a/Engine/Plugins/Experimental/PlainProps/Source/Private/PlainPropsPrint.cpp b/Engine/Plugins/Experimental/PlainProps/Source/Private/PlainPropsPrint.cpp
--- a/Engine/Plugins/Experimental/PlainProps/Source/Private/PlainPropsPrint.cpp
+++ b/Engine/Plugins/Experimental/PlainProps/Source/Private/PlainPropsPrint.cpp
@@ -1,6 +1,7 @@
 // Copyright Epic Games, Inc. All Rights Reserved.
 
 #include "PlainPropsPrint.h"
+#include "PlainPropsDiff.h"
 #include "PlainPropsIndex.h"
 #include "Containers/StringConv.h"
 #include "Containers/Utf8String.h"
@@ -1225,4 +1226,40 @@ FString FDebugIds::Print(FStructSchemaId Name) const
 
 ///////////////////////////////////////////////////////////////////////////////
 
+void PrintDiff(FUtf8Builder& Out, const FDiffPath& Diff, const FDebugIds& Names)
+{
+	check(Diff.Num());
+
+	for (FDiffNode Node : ReverseIterate(Diff))
+	{
+		Names.AppendString(Out, Node.Name);
+		Out << '.';
+	}
+	Out.RemoveSuffix(1);
+	Out << ' ';
+	Out << '(';
+	for (FDiffNode Node : ReverseIterate(Diff))
+	{
+		if (Node.Type.IsStruct())
+		{
+			Names.AppendString(Out, Names.Resolve(Node.Meta.Struct).Name);		
+		}
+		else if (Node.Type.IsRange())
+		{
+			Out << Node.Meta.Range.GetId();	
+		}
+		else if (FOptionalEnumSchemaId Enum = Node.Meta.Leaf)
+		{
+			Names.AppendString(Out, Names.Resolve(Enum.Get()).Name);
+		}
+		else
+		{
+			Out << ToString(ToLeafType(Node.Type.AsLeaf()));
+		}
+		Out << ' ';
+	}
+	Out.RemoveSuffix(1);
+	Out << ')';
+}
+
 } // namespace PlainProps
diff --git a/Engine/Plugins/Experimental/PlainProps/Source/Private/PlainPropsSave.cpp b/Engine/Plugins/Experimental/PlainProps/Source/Private/PlainPropsSave.cpp
--- a/Engine/Plugins/Experimental/PlainProps/Source/Private/PlainPropsSave.cpp
+++ b/Engine/Plugins/Experimental/PlainProps/Source/Private/PlainPropsSave.cpp
@@ -4,6 +4,7 @@
 #include "PlainPropsBind.h"
 #include "PlainPropsInternalBind.h"
 #include "PlainPropsInternalBuild.h"
+#include "PlainPropsInternalDiff.h"
 #include "PlainPropsInternalFormat.h"
 #include "PlainPropsSaveMember.h"
 #include "Algo/Find.h"
@@ -211,12 +212,6 @@ FBuiltRange* SaveLeafRange(const void* Range, const ILeafRangeBinding& Binding,
 	}
 }
 
-static ELeafWidth GetBasicWidth(FLeafBindType Leaf)
-{
-	checkf(Leaf.Bind.Type != ELeafBindType::BitfieldBool, TEXT("Arrays of bitfields is not a thing"));
-	return Leaf.Basic.Width;
-}
-
 [[nodiscard]] FBuiltRange* SaveRange(const void* Range, FRangeMemberBinding Member, const FSaveContext& Ctx)
 {
 	FRangeBinding Binding = Member.RangeBindings[0];
@@ -287,11 +282,6 @@ static ELeafWidth GetBasicWidth(FLeafBindType Leaf)
 	return { FMemberType(Member.RangeBindings[0].GetSizeType()), ToMemberType(Member.InnerTypes[0]), Member.NumRanges, Member.InnermostSchema, InnerRangeTypes };
 }
 
-static const uint8* At(const void* Ptr, SIZE_T Offset)
-{
-	return static_cast<const uint8*>(Ptr) + Offset;
-}
-
 static void SaveMember(FMemberBuilder& Out, const void* Struct, FMemberId Name, const FSaveContext& Ctx, FLeafMemberBinding Member)
 {
 	FUnpackedLeafType Type = { ToLeafType(Member.Leaf.Type), Member.Leaf.Width };
@@ -342,123 +332,6 @@ FBuiltStruct* SaveStruct(const void* Struct, FStructSchemaId BindId, const FSave
 
 ////////////////////////////////////////////////////////////////////////////////////////////////
 
-static bool DiffItem(const void* A, const void* B, const FSaveContext& Ctx, FRangeMemberBinding Range);
-static bool DiffItem(const void* A, const void* B, const FSaveContext& Ctx, FStructSchemaId Struct);
-
-static int64 DiffLeaf(const uint8* A, const uint8* B, FUnpackedLeafBindType Leaf)
-{
-	if (Leaf.Type == ELeafBindType::BitfieldBool)
-	{
-		uint32 Mask = 1u << Leaf.BitfieldIdx;
-		return (*A & Mask) - (*B & Mask);
-	}
-
-	switch (Leaf.Width)
-	{
-	case ELeafWidth::B8:	return FMemory::Memcmp(A, B, 1);
-	case ELeafWidth::B16:	return FMemory::Memcmp(A, B, 2);
-	case ELeafWidth::B32:	return FMemory::Memcmp(A, B, 4);
-	case ELeafWidth::B64:	return FMemory::Memcmp(A, B, 8);
-	}
-
-	check(false);
-	return 0;
-}
-
-static bool DiffItemSlice(const uint8* A, const uint8* B, uint64 Num, uint32 Stride, const FSaveContext&, SIZE_T LeafSize)
-{
-	return !!FMemory::Memcmp(A, B, Num * Stride);
-}
-
-template<typename ItemType>
-static bool DiffItemSlice(const uint8* A, const uint8* B, uint64 Num, uint32 Stride, const FSaveContext& Ctx, ItemType&& Member)
-{
-	for (const uint8* EndA = A + Num * Stride; A != EndA; A += Stride, B += Stride)
-	{
-		if (DiffItem(A, B, Ctx, Member))
-		{
-			return true;
-		}
-	}
-
-	return false;
-}
-
-struct FItemRangeReader
-{
-	explicit FItemRangeReader(const void* Range, const IItemRangeBinding& Binding)
-	: Ctx{{Range}}
-	{
-		ReadItems(Binding);
-	}
-
-	FSaveRangeContext Ctx;
-	const uint8* SliceIt = nullptr;
-	uint64 SliceNum = 0;
-
-	void ReadItems(const IItemRangeBinding& Binding)
-	{
-		Binding.ReadItems(Ctx);
-		SliceIt = static_cast<const uint8*>(Ctx.Items.Slice.Data);
-		SliceNum = Ctx.Items.Slice.Num;
-	}
-
-	void RefillItems(const IItemRangeBinding& Binding)
-	{
-		if (SliceNum == 0)
-		{
-			ReadItems(Binding);
-			check(SliceNum > 0);
-		}
-	}
-
-	const uint8* GrabItems(uint64 Num, uint32 Stride)
-	{
-		const uint8* Out = SliceIt;
-		SliceIt += Num * Stride; 
-		SliceNum -= Num;
-		return Out;
-	}
-};
-
-template<class T>
-bool DiffItemRange(const void* RangeA, const void* RangeB, const IItemRangeBinding& Binding, const FSaveContext& OuterCtx, T&& ItemSchema)
-{
-	FItemRangeReader A(RangeA, Binding);
-	FItemRangeReader B(RangeB, Binding);
-	if (A.Ctx.Items.NumTotal != B.Ctx.Items.NumTotal)
-	{
-		return true;
-	}
-
-	if (const uint64 NumTotal = A.Ctx.Items.NumTotal)
-	{
-		check(A.Ctx.Items.Stride == B.Ctx.Items.Stride);
-		const uint32 Stride = A.Ctx.Items.Stride;
-		while (true)
-		{
-			uint64 Num = FMath::Min(A.SliceNum, B.SliceNum);
-			if (DiffItemSlice(A.GrabItems(Num, Stride), B.GrabItems(Num, Stride), Num, Stride, OuterCtx, ItemSchema))
-			{
-				return true;
-			}
-			else if (A.Ctx.Request.NumRead + Num >= NumTotal)
-			{
-				check(A.Ctx.Request.NumRead + Num == NumTotal);	
-				check(B.Ctx.Request.NumRead + Num == NumTotal);	
-				return false;
-			}
-			
-			A.Ctx.Request.NumRead += Num;
-			B.Ctx.Request.NumRead += Num;
-			A.RefillItems(Binding);
-			B.RefillItems(Binding);
-		}
-	}
-
-	return false;
-}
-
 static bool DiffItem(const uint8* A, const uint8* B, const FSaveContext& Ctx, FRangeMemberBinding Member)
 {
 	FRangeBinding Binding = Member.RangeBindings[0];
diff --git a/Engine/Plugins/Experimental/PlainProps/Source/Public/PlainPropsBind.h b/Engine/Plugins/Experimental/PlainProps/Source/Public/PlainPropsBind.h
--- a/Engine/Plugins/Experimental/PlainProps/Source/Public/PlainPropsBind.h
+++ b/Engine/Plugins/Experimental/PlainProps/Source/Public/PlainPropsBind.h
@@ -481,6 +481,7 @@ struct alignas(16) IItemRangeBinding
 	// Try living without virtual destructor and trivially destructible members
 	virtual void ReadItems(FSaveRangeContext& Ctx) const = 0;
 	virtual void MakeItems(FLoadRangeContext& Ctx) const = 0;
+	virtual const char* GetId() const = 0;
 };
 
 ////////////////////////////////////////////////////////////////////////////////////////////////
@@ -519,10 +520,10 @@ public:
 // Specialized binding for transcoding leaf ranges
 struct alignas(16) ILeafRangeBinding
 {
-	virtual void	SaveLeaves(const void* Range, FLeafRangeAllocator& Out) const = 0;
-	virtual void	LoadLeaves(void* Range, FLeafRangeLoadView Leaves) const = 0;
-	virtual bool	DiffLeaves(const void* RangeA, const void* RangeB) const = 0;
-
+	virtual void		SaveLeaves(const void* Range, FLeafRangeAllocator& Out) const = 0;
+	virtual void		LoadLeaves(void* Range, FLeafRangeLoadView Leaves) const = 0;
+	virtual bool		DiffLeaves(const void* RangeA, const void* RangeB) const = 0;
+	virtual const char* GetId() const = 0;
 	template<typename SizeType>
 	static bool Diff(SizeType NumA, SizeType NumB, const void* A, const void* B, SIZE_T ItemSize)
 	{
@@ -547,7 +548,7 @@ public:
 	const IItemRangeBinding&	AsItemBinding() const		{ check(!IsLeafBinding()); return *reinterpret_cast<IItemRangeBinding*>(Handle & BindMask); }
 	const ILeafRangeBinding&	AsLeafBinding() const		{ check( IsLeafBinding()); return *reinterpret_cast<ILeafRangeBinding*>(Handle & BindMask); }
 	ERangeSizeType				GetSizeType() const			{ return static_cast<ERangeSizeType>(Handle & SizeMask); }
-
+	const char*					GetId() const				{ return IsLeafBinding() ? AsLeafBinding().GetId() : AsItemBinding().GetId(); }
 	inline bool	operator==(FRangeBinding O) const { return Handle == O.Handle; }
 };
 
diff --git a/Engine/Plugins/Experimental/PlainProps/Source/Public/PlainPropsDiff.h b/Engine/Plugins/Experimental/PlainProps/Source/Public/PlainPropsDiff.h
new file mode 100644
--- /dev/null
+++ b/Engine/Plugins/Experimental/PlainProps/Source/Public/PlainPropsDiff.h
@@ -0,0 +1,39 @@
+// Copyright Epic Games, Inc. All Rights Reserved.
+
+#pragma once
+
+#include "PlainPropsBind.h"
+
+namespace PlainProps 
+{
+
+union FDiffMetadata
+{
+	FOptionalEnumSchemaId	Leaf;
+	FRangeBinding			Range;
+	FStructSchemaId			Struct;
+};
+
+// Currently lacking range indices
+struct FDiffNode
+{
+	FMemberBindType			Type;
+	FOptionalMemberId		Name;
+	FDiffMetadata			Meta;
+	const void*				A;
+	const void*				B;
+};
+
+struct FDiffPath : public TArray<FDiffNode, TInlineAllocator<16>> {};
+
+struct FDiffContext
+{
+	const FDeclarations&	Types;
+	const FSchemaBindings&	Schemas;
+	const FCustomBindings&	Customs;
+	FDiffPath				Out;
+};
+
+[[nodiscard]] PLAINPROPS_API bool DiffStructs(FStructSchemaId BindId, const void* A, const void* B, FDiffContext& Ctx);
+
+} // namespace PlainProps
\ No newline at end of file
diff --git a/Engine/Plugins/Experimental/PlainProps/Source/Public/PlainPropsPrint.h b/Engine/Plugins/Experimental/PlainProps/Source/Public/PlainPropsPrint.h
--- a/Engine/Plugins/Experimental/PlainProps/Source/Public/PlainPropsPrint.h
+++ b/Engine/Plugins/Experimental/PlainProps/Source/Public/PlainPropsPrint.h
@@ -10,6 +10,7 @@ namespace PlainProps
 
 class FReadBatchId;
 struct FStructView;
+struct FDiffPath;
 
 ///////////////////////////////////////////////////////////////////////////////
 
@@ -19,18 +20,20 @@ PLAINPROPS_API FAnsiStringView ToString(FUnpackedLeafType Leaf);
 ///////////////////////////////////////////////////////////////////////////////
 
 PLAINPROPS_API void PrintMemberSchema(
-	FUtf8StringBuilderBase& Out,
+	FUtf8Builder& Out,
 	const FDebugIds& Names, 
 	FMemberType Type,
 	FOptionalSchemaId InnerSchema,
 	TConstArrayView<FMemberType> InnerRangeTypes);
 
 void PrintBatchAsYamlDoc(
-	FUtf8StringBuilderBase& Out,
+	FUtf8Builder& Out,
 	const FDebugIds& Names,
 	FReadBatchId Batch,
 	TConstArrayView<FStructView> Objects);
 
+PLAINPROPS_API void PrintDiff(FUtf8Builder& Out, const FDiffPath& Diff, const FDebugIds& Names);
+
 ///////////////////////////////////////////////////////////////////////////////
 
 } // namespace PlainProps
diff --git a/Engine/Plugins/Experimental/PlainProps/Source/Public/PlainPropsTypes.h b/Engine/Plugins/Experimental/PlainProps/Source/Public/PlainPropsTypes.h
--- a/Engine/Plugins/Experimental/PlainProps/Source/Public/PlainPropsTypes.h
+++ b/Engine/Plugins/Experimental/PlainProps/Source/Public/PlainPropsTypes.h
@@ -414,6 +414,8 @@ enum class ESchemaFormat { StableNames, InMemoryNames };
 
 //////////////////////////////////////////////////////////////////////////
 
+using FUtf8Builder = FUtf8StringBuilderBase;
+
 class FDebugIds
 {
 public:
@@ -424,15 +426,15 @@ public:
 	virtual FTypeId						Resolve(FEnumSchemaId Id) const = 0;
 	virtual FTypeId						Resolve(FStructSchemaId Id) const = 0;
 
-	virtual void						AppendString(FUtf8StringBuilderBase& Out, FNameId Name) const = 0;
+	virtual void						AppendString(FUtf8Builder& Out, FNameId Name) const = 0;
 
-	PLAINPROPS_API virtual void			AppendString(FUtf8StringBuilderBase& Out, FMemberId Name) const;
-	PLAINPROPS_API virtual void			AppendString(FUtf8StringBuilderBase& Out, FOptionalMemberId Name) const;
-	PLAINPROPS_API virtual void			AppendString(FUtf8StringBuilderBase& Out, FScopeId Scope) const;
-	PLAINPROPS_API virtual void			AppendString(FUtf8StringBuilderBase& Out, FTypenameId Typename) const;
-	PLAINPROPS_API virtual void			AppendString(FUtf8StringBuilderBase& Out, FTypeId Type) const;
-	PLAINPROPS_API virtual void			AppendString(FUtf8StringBuilderBase& Out, FEnumSchemaId Name) const;
-	PLAINPROPS_API virtual void			AppendString(FUtf8StringBuilderBase& Out, FStructSchemaId Name) const;
+	PLAINPROPS_API virtual void			AppendString(FUtf8Builder& Out, FMemberId Name) const;
+	PLAINPROPS_API virtual void			AppendString(FUtf8Builder& Out, FOptionalMemberId Name) const;
+	PLAINPROPS_API virtual void			AppendString(FUtf8Builder& Out, FScopeId Scope) const;
+	PLAINPROPS_API virtual void			AppendString(FUtf8Builder& Out, FTypenameId Typename) const;
+	PLAINPROPS_API virtual void			AppendString(FUtf8Builder& Out, FTypeId Type) const;
+	PLAINPROPS_API virtual void			AppendString(FUtf8Builder& Out, FEnumSchemaId Name) const;
+	PLAINPROPS_API virtual void			AppendString(FUtf8Builder& Out, FStructSchemaId Name) const;
 
 	PLAINPROPS_API FString				Print(FNameId Name) const;
 	PLAINPROPS_API FString				Print(FMemberId Name) const;
diff --git a/Engine/Plugins/Experimental/PlainProps/Source/Public/PlainPropsUeCoreBindings.h b/Engine/Plugins/Experimental/PlainProps/Source/Public/PlainPropsUeCoreBindings.h
--- a/Engine/Plugins/Experimental/PlainProps/Source/Public/PlainPropsUeCoreBindings.h
+++ b/Engine/Plugins/Experimental/PlainProps/Source/Public/PlainPropsUeCoreBindings.h
@@ -115,7 +115,7 @@ namespace PlainProps::UE
 {
 
 template <typename T, class Allocator>
-struct TArrayBinding : public IItemRangeBinding
+struct TArrayBinding : IItemRangeBinding
 {
 	using SizeType = int32;
 	using ItemType = T;
@@ -142,6 +142,8 @@ struct TArrayBinding : public IItemRangeBinding
 		const ArrayType& Array = Ctx.Request.GetRange<ArrayType>();
 		Ctx.Items.SetAll(Array.GetData(), static_cast<uint64>(Array.Num()));
 	}
+
+	virtual const char* GetId() const override { return "Array"; }
 };
 
 //////////////////////////////////////////////////////////////////////////
@@ -208,12 +210,14 @@ struct TStringBinding : ILeafRangeBinding
 		// Case-sensitive unnormalized comparison
 		return Diff(A.Len(), B.Len(), A.GetCharArray().GetData(), B.GetCharArray().GetData(), sizeof(TCHAR));
 	}
+
+	virtual const char* GetId() const override { return "String"; }
 };
 
 //////////////////////////////////////////////////////////////////////////
 
 template <typename T>
-struct TUniquePtrBinding : public IItemRangeBinding
+struct TUniquePtrBinding : IItemRangeBinding
 {
 	using SizeType = bool;
 	using ItemType = T;
@@ -249,12 +253,14 @@ struct TUniquePtrBinding : public IItemRangeBinding
 		const TUniquePtr<T>& Ptr = Ctx.Request.GetRange<TUniquePtr<T>>();
 		Ctx.Items.SetAll(Ptr.Get(), Ptr ? 1 : 0);
 	}
+
+	virtual const char* GetId() const override { return "UniquePtr"; }
 };
 
 //////////////////////////////////////////////////////////////////////////
 
 template <typename T>
-struct TOptionalBinding : public IItemRangeBinding
+struct TOptionalBinding : IItemRangeBinding
 {
 	using SizeType = bool;
 	using ItemType = T;
@@ -301,12 +307,14 @@ struct TOptionalBinding : public IItemRangeBinding
 		check(!Opt || reinterpret_cast<const T*>(&Opt) == &Opt.GetValue());
 		Ctx.Items.SetAll(reinterpret_cast<const T*>(Opt ? &Opt : nullptr), Opt ? 1 : 0);
 	}
+	
+	virtual const char* GetId() const override { return "Optional"; }
 };
 
 //////////////////////////////////////////////////////////////////////////
 
 template <typename T, typename KeyFuncs, typename SetAllocator>
-struct TSetBinding : public IItemRangeBinding
+struct TSetBinding : IItemRangeBinding
 {
 	using SizeType = int32;
 	using ItemType = T;
@@ -409,6 +417,8 @@ struct TSetBinding : public IItemRangeBinding
 		for (; Elems.IsValidIndex(Idx + Num); ++Num) {}
 		return { &Elems[Idx], static_cast<uint64>(Num) };
 	}
+
+	virtual const char* GetId() const override { return "Set"; }
 };
 
 //////////////////////////////////////////////////////////////////////////
diff --git a/Engine/Plugins/Experimental/PlainPropsUObject/Source/Private/PlainPropsUObjectBindings.cpp b/Engine/Plugins/Experimental/PlainPropsUObject/Source/Private/PlainPropsUObjectBindings.cpp
--- a/Engine/Plugins/Experimental/PlainPropsUObject/Source/Private/PlainPropsUObjectBindings.cpp
+++ b/Engine/Plugins/Experimental/PlainPropsUObject/Source/Private/PlainPropsUObjectBindings.cpp
@@ -3,6 +3,8 @@
 #include "PlainPropsUObjectRuntime.h"
 #include "PlainPropsUeCoreBindings.h"
 #include "PlainPropsBuildSchema.h"
+#include "PlainPropsDiff.h"
+#include "PlainPropsPrint.h"
 #include "PlainPropsWrite.h"
 #include "Algo/Compare.h"
 #include "Containers/PagedArray.h"
@@ -333,9 +335,30 @@ static bool ShouldBind(const UStruct* Struct)
 	return false;
 }
 
+// Serialize property-less UObject and UScriptStructs as their super class.
+//
+// E.g. FVector_NetQuantize10 is a pure runtime abstraction serialized as FVector.
+//		FAttenuationSubmixSendSettings is just a FSoundSubmixSendInfoBase but has a different 
+//		default constructor that matters in sparse delta serialization.
+//		UObjects are never instantiated during serialization so can be safely simplified.
+//
+// These heuristics might need more tuning
 static const UStruct* SkipEmptyBases(const UStruct* In)
 {
-	return In->PropertyLink ? In->PropertyLink->GetOwnerChecked<UStruct>() : In;
+	const UStruct* FirstOwner = In->PropertyLink ? In->PropertyLink->GetOwnerChecked<UStruct>() : In;
+	if (In != FirstOwner)
+	{
+		if (const UScriptStruct* Struct = Cast<const UScriptStruct>(In))
+		{
+			if (!(Struct->StructFlags & STRUCT_ZeroConstructor))
+			{
+				return In;
+			}
+		}
+
+		return FirstOwner;
+	}
+	return In;
 }
 
 static FTypeId IndexType(const UField* Field)
@@ -993,6 +1016,8 @@ struct TTrivialArrayBinding : IItemRangeBinding
 
 		Ctx.Items.Set(Array.GetData(), static_cast<uint64>(NewNum), Info.ItemSize);
 	}
+
+	virtual const char* GetId() const override { return "TrivialArray";}
 };
 
 // Currently can't extract constructor/destructor function pointers from FProperty, which
@@ -1027,6 +1052,8 @@ struct TNonTrivialArrayBinding : TTrivialArrayBinding<ScriptArray>
 		Ctx.Items.Set(Array.GetData(), static_cast<uint64>(NewNum), Info.ItemSize);
 	}
 
+	virtual const char* GetId() const override { return "NonTrivialArray";}
+
 	void InitItems(int32 Num, uint8* Items) const
 	{
 		if (Info.bConstructor)
@@ -1077,6 +1104,8 @@ struct TLeafArrayBinding : ILeafRangeBinding
 		return Diff(A.Num(), B.Num(), A.GetData(), B.GetData(), LeafSize);
 	}
 
+	virtual const char* GetId() const override { return "LeafArray";}
+
 	inline constexpr uint64 NumBytes(int32 NumItems) const
 	{
 		return static_cast<uint64>(NumItems) * LeafSize;
@@ -1807,6 +1836,8 @@ struct TSetPropertyBinding : IItemRangeBinding, ICustomBinding
 			Rehash(Set);
 		}
 	}
+
+	virtual const char* GetId() const override { return "Set"; }
 	
 	virtual void SaveCustom(FMemberBuilder& Dst, const void* InSrc, const void* InDefault, const FSaveContext& Ctx) override
 	{
@@ -2137,6 +2168,8 @@ struct TMapPropertyBinding : IItemRangeBinding
 			Rehash(Map);
 		}
 	}
+
+	virtual const char* GetId() const override { return "Map";}
 };
 
 ////////////////////////////////////////////////////////////////////////////////////////////////
@@ -2580,6 +2613,9 @@ struct FStaticArrayBinding : IItemRangeBinding
 	{
 		Ctx.Items.Set(&Ctx.Request.GetRange<uint8>(), Num, Stride);
 	}
+
+	
+	virtual const char* GetId() const override { return "StaticArray";}
 };
 
 ////////////////////////////////////////////////////////////////////////////////////////////////
@@ -2982,7 +3018,7 @@ struct FMemoryPropertyBatch
 
 	static void Save(FMemberBuilder& Out, FName In, const FSaveContext&)
 	{
-		Out.Add(GUE.Members.Id, In.ToUnstableInt());
+		Out.Add(GUE.Members.Id, ToInt(In));
 	}
 
 	static void Load(FName& Out, FStructLoadView In)
@@ -3156,13 +3192,9 @@ public:
 		SavedObjects.Reserve(NumReserve);
 	}
 
-	void Save(const UObject* Object)
+	void Save(FStructSchemaId Id, const UObject* Object, const UObject* CDO)
 	{
-		const UClass* Class = Object->GetClass();
-		const UObject* CDO = Class->GetDefaultObject(/* create lazily */ false);
-		FStructSchemaId Id = GUE.Names.IndexStruct(IndexType(Class));
-		bool bFlat = Object == CDO || CDO == nullptr;
-		SavedObjects.Emplace(Id, bFlat ? SaveStruct(Object, Id, FlatCtx) : SaveStructDelta(Object, CDO, Id, DeltaCtx));
+		SavedObjects.Emplace(Id, CDO ? SaveStructDelta(Object, CDO, Id, DeltaCtx) : SaveStruct(Object, Id, FlatCtx));
 	}
 
 	TArray64<uint8>	Write(TArray<FStructSchemaId>& OutRuntimeIds) const
@@ -3274,6 +3306,12 @@ public:
 		++LoadIdx;
 	}
 
+	void Reload(UObject* Dst, int32 ReloadIdx)
+	{
+		FStructView In = Objects[ReloadIdx];
+		LoadStruct(Dst, In.Values, In.Schema.Id, *Plans);
+	}
+
 private:
 	FLoadBatchPtr				Plans;
 	TArray<FStructView>			Objects;
@@ -3292,10 +3330,33 @@ struct FInstance
 	UObject* Original;
 	UObject* PP;
 	UObject* TPS;
+	const UObject* CDO;
+	FStructSchemaId Id;
+
+	void Init()
+	{
+		const UClass* Class = Original->GetClass();
+		CDO = Class->GetDefaultObject(/* create lazily */ true);
+		CDO = CDO == Original ? nullptr : CDO;
+		Id = GUE.Names.IndexStruct(IndexType(Class));
+	}
 };
 
+static UObject* MakeEmptyInstance(UObject* Obj, FName Name)
+{
+	FStaticConstructObjectParameters Params(Obj->GetClass());
+	Params.Outer = Obj->GetOuter();
+	Params.Name = Name;
+	Params.SetFlags = Obj->GetFlags();
+	Params.Template = Obj->GetArchetype();
+	Params.bAssumeTemplateIsArchetype = true;
+	Params.bCopyTransientsFromClassDefaults = true;
+	return StaticConstructObject_Internal(Params);
+}
+
 static void Roundtrip()
 {
+	// Gather all non-empty objects
 	TArray<FInstance> Instances;
 	for (TObjectIterator<UObject> It; It; ++It)
 	{
@@ -3305,39 +3366,61 @@ static void Roundtrip()
 			Instances.Add({*It});
 		}
 	}
+
+	// Create CDOs if needed and then clones for PP and TPS tests
+	for (FInstance& Instance : Instances)
+	{
+		Instance.Init();
+	}
 	FName PP_N("PP", 1);
 	for (FInstance& Instance : Instances)
 	{
-		//todo: duplicate CDO somehow
-		UObject* O = Instance.Original;
 		PP_N.SetNumber(1 + &Instance - Instances.GetData());
-		Instance.PP = StaticAllocateObject(O->GetClass(), O->GetOuter(), PP_N, O->GetFlags());
+		Instance.PP = MakeEmptyInstance(Instance.Original, PP_N);
 	}
 	FName TPS_N("TPS", 1);
 	for (FInstance& Instance : Instances)
 	{
-		UObject* O = Instance.Original;
 		TPS_N.SetNumber(1 + &Instance - Instances.GetData());
-		Instance.TPS = StaticAllocateObject(O->GetClass(), O->GetOuter(), TPS_N, O->GetFlags());
+		Instance.TPS = MakeEmptyInstance(Instance.Original, TPS_N);
 	}
 
+	// Save
 	FMemoryBatch Doc;
 	TCustomPropertyBindings<FMemoryPropertyBatch> FlatBindings(Doc.Properties, GUE.Customs);
 	TCustomPropertyBindings<FMemoryPropertyBatch> DeltaBindings(Doc.Properties, GUE.DeltaCustoms);
 	{
 		FBatchSaver Batch(FlatBindings.Overlay, DeltaBindings.Overlay, GUObjectArray.GetObjectArrayNum());
 		for (const FInstance& Instance : Instances)
-		{
-			Batch.Save(Instance.Original);
+		{	
+			Batch.Save(Instance.Id, Instance.Original, Instance.CDO);
 		}
 		Doc.Data = Batch.Write(/* out */ Doc.RuntimeIds);
 	}
 
+	// Load
 	FMemoryBatchLoader Batch(DeltaBindings.Overlay, MakeMemoryView(Doc.Data), Doc.RuntimeIds);
 	for (const FInstance& Instance : Instances)
 	{
 		Batch.Load(Instance.PP);
 	}
+
+	// Diff original vs PlainProps
+	FDiffContext FlatCtx	= { GUE.Types, GUE.Schemas, FlatBindings.Overlay };
+	FDiffContext DeltaCtx	= { GUE.Types, GUE.Schemas, DeltaBindings.Overlay };
+	FDiffPath Path;
+	TUtf8StringBuilder<256> Diffs;
+	for (FInstance& Instance : Instances)
+	{
+		FDiffContext& Ctx = Instance.CDO ? DeltaCtx : FlatCtx;
+		if (DiffStructs(Instance.Id, Instance.Original, Instance.PP, /* in-out */ Ctx))
+		{
+			PrintDiff(/* out */ Diffs, Ctx.Out, GUE.Names);
+			//Batch.Reload(Instance.TPS, &Instance - Instances.GetData());
+			Ctx.Out.Reset();
+			Diffs.Append("\n");	
+		}
+	}
 }
 
 static void TestBindings()
