PlainProps Make TArray and TOptional item binding destroy existing elements rather than reusing them to fix issues with CDO + delta loading

#rnx

diff --git a/Engine/Plugins/Experimental/PlainProps/Source/Public/PlainPropsUeCoreBindings.h b/Engine/Plugins/Experimental/PlainProps/Source/Public/PlainPropsUeCoreBindings.h
--- a/Engine/Plugins/Experimental/PlainProps/Source/Public/PlainPropsUeCoreBindings.h
+++ b/Engine/Plugins/Experimental/PlainProps/Source/Public/PlainPropsUeCoreBindings.h
@@ -268,36 +268,27 @@ struct TOptionalBinding : IItemRangeBinding
 	virtual void MakeItems(FLoadRangeContext& Ctx) const override
 	{
 		TOptional<T>& Opt = Ctx.Request.GetRange<TOptional<T>>();
-		
+		Opt.Reset();
+
 		if (Ctx.Request.NumTotal() == 0)
 		{
-			Opt.Reset();
+			return;
 		}
-		else if constexpr (std::is_default_constructible_v<T>)
+		
+		if constexpr (std::is_default_constructible_v<T>)
 		{
-			if (!Opt)
-			{
-				Opt.Emplace();
-			}
+			Opt.Emplace();
 			Ctx.Items.Set(reinterpret_cast<T*>(&Opt), 1);
 		}
-		else if (Opt)
+		else if (Ctx.Request.IsFirstCall())
 		{
-			Ctx.Items.Set(reinterpret_cast<T*>(&Opt), 1);
+			Ctx.Items.SetUnconstructed();
+			Ctx.Items.RequestFinalCall();
+			Ctx.Items.Set(reinterpret_cast<T*>(&Opt), 1);	
 		}
-		else
+		else // Move-construct from self reference
 		{
-			if (Ctx.Request.IsFirstCall())
-			{
-				Ctx.Items.SetUnconstructed();
-				Ctx.Items.RequestFinalCall();
-				Ctx.Items.Set(reinterpret_cast<T*>(&Opt), 1);	
-			}
-			else
-			{
-				// Move-construct from self reference
-				Opt.Emplace(reinterpret_cast<T&&>(Opt));
-			}
+			Opt.Emplace(reinterpret_cast<T&&>(Opt));
 		}
 	}
 
diff --git a/Engine/Plugins/Experimental/PlainPropsUObject/Source/Private/PlainPropsUObjectBindings.cpp b/Engine/Plugins/Experimental/PlainPropsUObject/Source/Private/PlainPropsUObjectBindings.cpp
--- a/Engine/Plugins/Experimental/PlainPropsUObject/Source/Private/PlainPropsUObjectBindings.cpp
+++ b/Engine/Plugins/Experimental/PlainPropsUObject/Source/Private/PlainPropsUObjectBindings.cpp
@@ -7,6 +7,7 @@
 #include "PlainPropsPrint.h"
 #include "PlainPropsWrite.h"
 #include "Algo/Compare.h"
+#include "Algo/Find.h"
 #include "Containers/PagedArray.h"
 #include "Hash/xxhash.h"
 #include "Math/PreciseFP.h"
@@ -989,9 +990,9 @@ static_assert(sizeof(FArrayPropertyInfo) == 8);
 template<class ScriptArray>
 struct TTrivialArrayBinding : IItemRangeBinding
 {
-	TTrivialArrayBinding(FArrayPropertyInfo InFo) : Info(InFo) {}
+	const FArrayPropertyInfo Info;
 
-	FArrayPropertyInfo Info;
+	TTrivialArrayBinding(FArrayPropertyInfo InFo) : Info(InFo) {}
 
 	virtual void ReadItems(FSaveRangeContext& Ctx) const override
 	{
@@ -1004,16 +1005,12 @@ struct TTrivialArrayBinding : IItemRangeBinding
 		ScriptArray& Array = Ctx.Request.GetRange<ScriptArray>();
 
 		int32 NewNum = static_cast<int32>(Ctx.Request.NumTotal());
-		int32 Diff = NewNum - Array.Num();
-		if (Diff > 0)
-		{
-			Array.AddZeroed(Diff, Info.ItemSize, Info.ItemAlign);
-		}
-		else if (Diff < 0)
+		Array.SetNumUninitialized(NewNum, Info.ItemSize, Info.ItemAlign);
+		if (NewNum)
 		{
-			Array.Remove(NewNum, -Diff, Info.ItemSize, Info.ItemAlign);
+			FMemory::Memzero(Array.GetData(), NumBytes(NewNum, Info.ItemSize));
 		}
-
+		
 		Ctx.Items.Set(Array.GetData(), static_cast<uint64>(NewNum), Info.ItemSize);
 	}
 
@@ -1025,54 +1022,44 @@ struct TTrivialArrayBinding : IItemRangeBinding
 template<class ScriptArray>
 struct TNonTrivialArrayBinding : TTrivialArrayBinding<ScriptArray>
 {
+	const FProperty* Inner;
+
 	using TTrivialArrayBinding<ScriptArray>::Info;
 
-	TNonTrivialArrayBinding(FArrayPropertyInfo InFo, FProperty* InNer)
+	TNonTrivialArrayBinding(FArrayPropertyInfo InFo, const FProperty* InNer)
 	: TTrivialArrayBinding<ScriptArray>(InFo)
 	, Inner(InNer)
 	{}
 
+	//todo: other range bindings
+
 	virtual void MakeItems(FLoadRangeContext& Ctx) const override
 	{
 		ScriptArray& Array = Ctx.Request.GetRange<ScriptArray>();
 
-		int32 NewNum = static_cast<int32>(Ctx.Request.NumTotal());
-		int32 Diff = NewNum - Array.Num();
-		if (Diff > 0)
-		{
-			int32 OldNum = Array.Add(Diff, Info.ItemSize, Info.ItemAlign);
-			InitItems(Diff, static_cast<uint8*>(Array.GetData()) + NumBytes(OldNum, Info.ItemSize));
-		}
-		else if (Diff < 0)
-		{
-			DestructItems(-Diff, static_cast<uint8*>(Array.GetData()) + NumBytes(NewNum, Info.ItemSize));
-			Array.Remove(NewNum, -Diff, Info.ItemSize, Info.ItemAlign);
-		}
+		int32 NumDestroy = Info.bDestructor * Array.Num();
+		DestroyValues(Inner, static_cast<uint8*>(Array.GetData()), NumDestroy, Info.ItemSize);
+
+		uint64 NewNum = Ctx.Request.NumTotal();
+		Array.SetNumUninitialized(static_cast<int32>(NewNum), Info.ItemSize, Info.ItemAlign);
+		InitItems(NewNum, Array.GetData());
 		
-		Ctx.Items.Set(Array.GetData(), static_cast<uint64>(NewNum), Info.ItemSize);
+		Ctx.Items.Set(Array.GetData(), NewNum, Info.ItemSize);
 	}
 
-	virtual const char* GetId() const override { return "NonTrivialArray";}
-
-	void InitItems(int32 Num, uint8* Items) const
+	inline void InitItems(uint64 Num, void* Items) const
 	{
 		if (Info.bConstructor)
 		{
-			ConstructValues(Inner, Items, Num, Info.ItemSize);
+			ConstructValues(Inner, static_cast<uint8*>(Items), Num, Info.ItemSize);
 		}
-		else
+		else if (Num)
 		{
-			FMemory::Memzero(Items, NumBytes(Num, Info.ItemSize));
+			FMemory::Memzero(Items, Num * Info.ItemSize);
 		}
 	}
-
-	void DestructItems(int32 Num, uint8* Items) const
-	{
-		int32 NumDestroy = Info.bDestructor * Num;
-		DestroyValues(Inner, Items, NumDestroy, Info.ItemSize);
-	}
-
-	FProperty* Inner;
+	
+	virtual const char* GetId() const override { return "NonTrivialArray";}
 };
 
 template<ELeafType Type, ELeafWidth Width>
@@ -2529,6 +2516,7 @@ static FPairBindings GPairs;
 
 //////////////////////////////////////////////////////////////////////////////////////////////
 
+// Todo: change to bind as a uni-range instead of struct, see TOptionalBinding
 class FOptionalBindings
 {
 	TMap<FParameterBinding, FStructSchemaId>	Bindings;
@@ -2599,11 +2587,11 @@ static const FStringBindings GStrings;
 
 struct FStaticArrayBinding : IItemRangeBinding
 {
-	FStaticArrayBinding(uint32 InNum, uint32 InStride) : Num(InNum), Stride(InStride) {}
-
 	uint32 Num;
 	uint32 Stride;
 
+	FStaticArrayBinding(uint32 InNum, uint32 InStride) : Num(InNum), Stride(InStride) {}
+
 	virtual void ReadItems(FSaveRangeContext& Ctx) const override
 	{
 		Ctx.Items.SetAll(Ctx.Request.Range, Num, Stride);
@@ -2613,9 +2601,8 @@ struct FStaticArrayBinding : IItemRangeBinding
 	{
 		Ctx.Items.Set(&Ctx.Request.GetRange<uint8>(), Num, Stride);
 	}
-
 	
-	virtual const char* GetId() const override { return "StaticArray";}
+	virtual const char* GetId() const override { return "StaticArray"; }
 };
 
 ////////////////////////////////////////////////////////////////////////////////////////////////
@@ -3354,6 +3341,11 @@ static UObject* MakeEmptyInstance(UObject* Obj, FName Name)
 	return StaticConstructObject_Internal(Params);
 }
 
+static bool Include(UClass* Class)
+{	
+	static const FName Exclusions[] = {"GameFeaturePluginStateMachine"};
+	return ShouldBind(Class) && !Algo::Find(Exclusions, Class->GetFName());
+}
 static void Roundtrip()
 {
 	// Gather all non-empty objects
@@ -3361,7 +3353,7 @@ static void Roundtrip()
 	for (TObjectIterator<UObject> It; It; ++It)
 	{
 		UObject* Object = *It;
-		if (ShouldBind(Object->GetClass()))
+		if (Include(Object->GetClass()))
 		{
 			Instances.Add({*It});
 		}
@@ -3421,6 +3413,8 @@ static void Roundtrip()
 			Diffs.Append("\n");	
 		}
 	}
+
+	UE_DEBUG_BREAK();
 }
 
 static void TestBindings()
