PlainProps reimplement FOptionalProperty struct binding as range binding

#rnx

diff --git a/Engine/Plugins/Experimental/PlainPropsUObject/Source/Private/PlainPropsUObjectBindings.cpp b/Engine/Plugins/Experimental/PlainPropsUObject/Source/Private/PlainPropsUObjectBindings.cpp
--- a/Engine/Plugins/Experimental/PlainPropsUObject/Source/Private/PlainPropsUObjectBindings.cpp
+++ b/Engine/Plugins/Experimental/PlainPropsUObject/Source/Private/PlainPropsUObjectBindings.cpp
@@ -900,17 +900,17 @@ inline uint64 NumBytes(int32 NumItems, SIZE_T ItemSize)
 	return static_cast<uint64>(NumItems) * ItemSize;
 }
 
-inline bool HasConstructor(FProperty* Property)
+inline bool HasConstructor(const FProperty* Property)
 {
 	return !(Property->PropertyFlags & CPF_ZeroConstructor);
 }
 
-inline bool HasDestructor(FProperty* Property)
+inline bool HasDestructor(const FProperty* Property)
 {
 	return !(Property->PropertyFlags & (CPF_IsPlainOldData | CPF_NoDestructor));
 }
 
-inline bool HasHash(FProperty* Property)
+inline bool HasHash(const FProperty* Property)
 {
 	return !!(Property->PropertyFlags & CPF_HasGetValueTypeHash);
 }
@@ -1014,7 +1014,7 @@ struct TTrivialArrayBinding : IItemRangeBinding
 		Ctx.Items.Set(Array.GetData(), static_cast<uint64>(NewNum), Info.ItemSize);
 	}
 
-	virtual const char* GetId() const override { return "TrivialArray";}
+	virtual const char* GetId() const override { return "TrivialArray"; }
 };
 
 // Currently can't extract constructor/destructor function pointers from FProperty, which
@@ -1059,7 +1059,7 @@ struct TNonTrivialArrayBinding : TTrivialArrayBinding<ScriptArray>
 		}
 	}
 	
-	virtual const char* GetId() const override { return "NonTrivialArray";}
+	virtual const char* GetId() const override { return "NonTrivialArray"; }
 };
 
 template<ELeafType Type, ELeafWidth Width>
@@ -1091,7 +1091,7 @@ struct TLeafArrayBinding : ILeafRangeBinding
 		return Diff(A.Num(), B.Num(), A.GetData(), B.GetData(), LeafSize);
 	}
 
-	virtual const char* GetId() const override { return "LeafArray";}
+	virtual const char* GetId() const override { return "LeafArray"; }
 
 	inline constexpr uint64 NumBytes(int32 NumItems) const
 	{
@@ -2516,42 +2516,169 @@ static FPairBindings GPairs;
 
 //////////////////////////////////////////////////////////////////////////////////////////////
 
-// Todo: change to bind as a uni-range instead of struct, see TOptionalBinding
-class FOptionalBindings
+inline void ReadBoolOptionalItem(FSaveRangeContext& Ctx, uint32 ItemSize)
+{
+	check((&Ctx.Request.GetRange<uint8>())[ItemSize] <= 1);
+	bool bSet = (&Ctx.Request.GetRange<bool>())[ItemSize];
+	Ctx.Items.SetAll(bSet ? Ctx.Request.Range : nullptr, uint64(bSet), ItemSize);
+}
+
+inline void MakeBoolOptionalItem(FLoadRangeContext& Ctx, uint32 ItemSize)
+{
+	check((&Ctx.Request.GetRange<uint8>())[ItemSize] <= 1);
+	bool& bSet = (&Ctx.Request.GetRange<bool>())[ItemSize];
+	bSet = Ctx.Request.NumTotal() > 0;
+	Ctx.Items.Set(&Ctx.Request.GetRange<uint8>(), uint64(bSet), ItemSize);
+}
+
+template<uint32 ItemSize>
+struct TTrivialOptionalBinding : IItemRangeBinding
 {
-	TMap<FParameterBinding, FStructSchemaId>	Bindings;
+	virtual void ReadItems(FSaveRangeContext& Ctx) const override { ReadBoolOptionalItem(Ctx, ItemSize); }
+	virtual void MakeItems(FLoadRangeContext& Ctx) const override { MakeBoolOptionalItem(Ctx, ItemSize); }
+	virtual const char* GetId() const override { return "TTrivialOptional"; }
+};
 
-	inline FStructSchemaId BindImpl(const FParameterBinding Value)
+struct FTrivialOptionalBinding : IItemRangeBinding
+{
+	const uint32 ItemSize;
+	explicit FTrivialOptionalBinding(uint32 Size) : ItemSize(Size) {}
+	virtual void ReadItems(FSaveRangeContext& Ctx) const override { ReadBoolOptionalItem(Ctx, ItemSize); }
+	virtual void MakeItems(FLoadRangeContext& Ctx) const override { MakeBoolOptionalItem(Ctx, ItemSize); }
+	virtual const char* GetId() const override { return "FTrivialOptional"; }
+};
+
+struct FIntrusiveOptionalBinding : IItemRangeBinding
+{
+	const FProperty*	Inner;
+	const uint32		ItemSize;
+
+	explicit FIntrusiveOptionalBinding(const FProperty* In) : Inner(In), ItemSize(In->GetElementSize()) {}
+	
+	virtual void ReadItems(FSaveRangeContext& Ctx) const override
+	{
+		bool bSet = Inner->IsIntrusiveOptionalValueSet(Ctx.Request.Range);
+		Ctx.Items.SetAll(bSet ? Ctx.Request.Range : nullptr, uint64(bSet), ItemSize);
+	}
+
+	virtual void MakeItems(FLoadRangeContext& Ctx) const override
 	{
-		check(Value.Offset == 0);
-		if (const FStructSchemaId* Id = Bindings.Find(Value))
+		uint8* Value = &Ctx.Request.GetRange<uint8>();
+		Inner->ClearIntrusiveOptionalValue(Value);
+		if (Ctx.Request.NumTotal() > 0)
 		{
-			return *Id;
+			Inner->InitializeIntrusiveUnsetOptionalValue(Value);
+			Ctx.Items.Set(Value, 1, ItemSize);
 		}
+		else
+		{
+			Ctx.Items.Set(nullptr, 0, ItemSize);
+		}
+	}
 
-		// Index names
-		FTypeId Params[1] = { Value.IndexParameterName(GUE.Names) };
-		FTypeId Type = { GUE.Scopes.Core, FTypenameId(GUE.Names.MakeParametricTypeId(GUE.Typenames.Optional, Params)) };
-		FStructSchemaId Id = GUE.Names.IndexStruct(Type);
-		FMemberId Members[1] = { GUE.Members.Value };
-		FMemberBinding Inners[1] = { Value };
+	virtual const char* GetId() const override { return "IntrusiveOptional"; }
+};
 
-		// Todo: Ownership / memory leak
-		GUE.Types.DeclareStruct(Id, Type, Members, EMemberPresence::AllowSparse);
-		GUE.Schemas.BindStruct(Id, Id, Inners);
+struct FNonIntrusiveOptionalBinding : IItemRangeBinding
+{
+	const FProperty*	Inner;
+	const uint32		ItemSize;
+	const bool			bConstructor;
+	const bool			bDestructor;
+	
+	explicit FNonIntrusiveOptionalBinding(const FProperty* In)
+	: Inner(In)
+	, ItemSize(In->GetElementSize())
+	, bConstructor(HasConstructor(In))
+	, bDestructor(HasDestructor(In))
+	{}
 
-		Bindings.Emplace(Value, Id);
+	virtual void ReadItems(FSaveRangeContext& Ctx) const override
+	{
+		ReadBoolOptionalItem(Ctx, ItemSize);
+	}
 
-		return Id;
+	virtual void MakeItems(FLoadRangeContext& Ctx) const override
+	{
+		uint8* Value = &Ctx.Request.GetRange<uint8>();
+		bool& bSet = reinterpret_cast<bool&>(Value[ItemSize]);
+		if (bDestructor && bSet)
+		{
+			DestroyValue(Inner, Value);
+		}
+
+		bSet = Ctx.Request.NumTotal() > 0;
+		Ctx.Items.Set(bSet ? Value : nullptr , 1, ItemSize);
+
+		if (bSet)
+		{
+			if (bConstructor)
+			{
+				ConstructValue(Inner, Value);
+			}
+			else
+			{
+				FMemory::Memzero(Value, ItemSize);
+			}
+			
+		}
 	}
-public:
 
-	inline FMemberBinding Bind(FMemberBinding Value)
+	virtual const char* GetId() const override { return "NonIntrusiveOptional"; }
+};
+
+
+
+class FOptionalBindings
+{
+	TTrivialOptionalBinding<1> Trivial1;
+	TTrivialOptionalBinding<2> Trivial2;
+	TTrivialOptionalBinding<4> Trivial4;
+	TTrivialOptionalBinding<8> Trivial8;
+	TTrivialOptionalBinding<12> Trivial12;
+	TTrivialOptionalBinding<16> Trivial16;
+	TTrivialOptionalBinding<24> Trivial24;
+	TTrivialOptionalBinding<32> Trivial32;
+
+	alignas(PLATFORM_CACHE_LINE_SIZE)
+	TMap<FParameterBinding, FRangeBinding>	Bindings;
+
+	const IItemRangeBinding* BindNew(FProperty* Inner)
 	{
-		FMemberBinding Out(0);
-		Out.InnermostSchema = BindImpl(static_cast<FParameterBinding>(Value));
-		Out.InnermostType = DefaultStructBindType;
-		return Out;
+		if (Inner->HasIntrusiveUnsetOptionalState())
+		{
+			return new FIntrusiveOptionalBinding(Inner);
+		}
+		else if (HasConstructor(Inner) || HasDestructor(Inner))
+		{
+			return new FNonIntrusiveOptionalBinding(Inner);
+		}
+		else switch (Inner->GetElementSize())
+		{
+			case 1:		return &Trivial1;
+			case 2:		return &Trivial2;
+			case 4:		return &Trivial4;
+			case 8:		return &Trivial8;
+			case 12:	return &Trivial12;
+			case 16:	return &Trivial16;
+			case 24:	return &Trivial24;
+			case 32:	return &Trivial32;
+			default:	return new FTrivialOptionalBinding(Inner->GetElementSize());
+		}
+	}
+
+public:
+	FRangeBinding Bind(FProperty* Inner, FMemberBinding Key)
+	{
+		check(Key.Offset == 0);
+		if (const FRangeBinding* Binding = Bindings.Find(FParameterBinding(Key)))
+		{
+			return *Binding;
+		}
+
+		// Todo: Ownership / memory leak
+		const IItemRangeBinding* Out = BindNew(Inner);
+		return Bindings.Emplace(Key, FRangeBinding(*Out, ERangeSizeType::Uni));
 	}
 };
 
@@ -2769,7 +2896,8 @@ private:
 	}
 	inline FMemberBinding BindOptional(FOptionalProperty* Property)
 	{
-		return GOptionals.Bind(BindSingleProperty(Property->GetValueProperty()));
+		FMemberBinding Inner = BindSingleProperty(Property->GetValueProperty());
+		return BindAsRange(Property, GOptionals.Bind(Property->GetValueProperty(), Inner), Inner);
 	}
 
 	#if WITH_VERSE_VM
