PlainProps move UObject roundtrip test into a new commandlet

Also some test improvements and a bug fix so we now roundtrip 240K UObjects without any diff

#rnx

diff --git a/Engine/Plugins/Experimental/PlainProps/Source/Public/PlainPropsUeCoreBindings.h b/Engine/Plugins/Experimental/PlainProps/Source/Public/PlainPropsUeCoreBindings.h
--- a/Engine/Plugins/Experimental/PlainProps/Source/Public/PlainPropsUeCoreBindings.h
+++ b/Engine/Plugins/Experimental/PlainProps/Source/Public/PlainPropsUeCoreBindings.h
@@ -381,7 +381,11 @@ struct TSetBinding : IItemRangeBinding
 		static_assert(offsetof(TSetElement<T>, Value) == 0);
 		const TSparseArray<TSetElement<T>>& Elems = Ctx.Request.GetRange<TSparseArray<TSetElement<T>>>();
 
-		if (FExistingItemSlice LastRead = Ctx.Items.Slice)
+		if (Elems.IsEmpty())
+		{
+			Ctx.Items.SetAll(nullptr, 0, sizeof(TSetElement<T>));
+		}
+		else if (FExistingItemSlice LastRead = Ctx.Items.Slice)
 		{
 			// Continue partial response
 			const TSetElement<T>* NextElem = static_cast<const TSetElement<T>*>(LastRead.Data) + LastRead.Num + /* skip known invalid */ 1;
@@ -389,8 +393,7 @@ struct TSetBinding : IItemRangeBinding
 		}
 		else if (Elems.IsCompact())
 		{
-			int32 Num = Elems.Num();
-			Ctx.Items.SetAll(Num ? &Elems[0] : nullptr, Num);
+			Ctx.Items.SetAll(&Elems[0], Elems.Num());
 		}
 		else
 		{
diff --git a/Engine/Plugins/Experimental/PlainPropsUObject/Source/PlainPropsUObject.Build.cs b/Engine/Plugins/Experimental/PlainPropsUObject/Source/PlainPropsUObject.Build.cs
--- a/Engine/Plugins/Experimental/PlainPropsUObject/Source/PlainPropsUObject.Build.cs
+++ b/Engine/Plugins/Experimental/PlainPropsUObject/Source/PlainPropsUObject.Build.cs
@@ -15,7 +15,8 @@ public class PlainPropsUObject : ModuleRules
 				"Core",
 				"CoreUObject",
 				"CorePreciseFP",
-				"PlainProps",
+				"Engine",
+				"PlainProps"
 		});
 	}
 }
diff --git a/Engine/Plugins/Experimental/PlainPropsUObject/Source/Private/PlainPropsUObjectBindings.cpp b/Engine/Plugins/Experimental/PlainPropsUObject/Source/Private/PlainPropsUObjectBindings.cpp
--- a/Engine/Plugins/Experimental/PlainPropsUObject/Source/Private/PlainPropsUObjectBindings.cpp
+++ b/Engine/Plugins/Experimental/PlainPropsUObject/Source/Private/PlainPropsUObjectBindings.cpp
@@ -14,6 +14,7 @@
 #include "Misc/CommandLine.h"
 #include "Misc/CoreDelegates.h"
 #include "Misc/DefinePrivateMemberPtr.h"
+#include "StructUtils/UserDefinedStruct.h"
 #include "UObject/AnsiStrProperty.h"
 #include "UObject/Class.h"
 #include "UObject/EnumProperty.h"
@@ -46,6 +47,9 @@ UE_DEFINE_PRIVATE_MEMBER_PTR(FName, GUnicastDelegateFunctionName, FUnicastScript
 UE_DEFINE_PRIVATE_MEMBER_PTR(FMulticastInvocationList, GMulticastDelegateInvocationList, FMulticastScriptDelegate, InvocationList);
 UE_DEFINE_PRIVATE_MEMBER_PTR(bool, GSparseDelegateIsBound, FSparseDelegate, bIsBound);
 
+
+DEFINE_LOG_CATEGORY_STATIC(LogPlainPropsUObject, Log, All);
+
 namespace PlainProps::UE
 {
 
@@ -137,7 +141,21 @@ void FDefaultStructs::Bind(FStructSchemaId Id, const UScriptStruct* Struct)
 	ReserveFlags(Id.Idx);
 	checkf(Flip(Bound[Id.Idx]), TEXT("'%s' already bound"), *GUE.Names.Print(Id));
 
-	if (!!(Flags & STRUCT_ZeroConstructor) || Ops == nullptr)
+	if (const UUserDefinedStruct* UserStruct = Cast<UUserDefinedStruct>(Struct))
+	{
+		const void* DefaultInstance = UserStruct->GetDefaultInstance();
+		check(DefaultInstance);
+		if (FMemory::MemIsZero(DefaultInstance, Size))
+		{
+			ReserveZeroes(Zeroes, Size, Alignment);
+		}
+		else
+		{
+			Instanced[Id.Idx] = true;
+			Instances.Emplace(Id, MakeStaticInstance(DefaultInstance));
+		}
+	}
+	else if (!!(Flags & STRUCT_ZeroConstructor) || Ops == nullptr)
 	{
 		ReserveZeroes(Zeroes, Size, Alignment);
 	}
@@ -1578,10 +1596,12 @@ uint8* SetNumUninitialized(ScriptType& Dst, const LayoutType& Layout, uint64 Num
 	return static_cast<uint8*>(Dst.GetData(0, Layout));
 }
 
+// @pre Elems.Num() > 0
 inline FExistingItemSlice GetContiguousSlice(int32 Idx, const FScriptSparseArray& Elems, const uint8* Data, SIZE_T Stride)
 {
+	checkSlow(!Elems.IsEmpty());
 	int32 Num = 1;
-	for (;!Elems.IsValidIndex(Idx); ++Idx) {}
+	for (;!Elems.IsValidIndex(Idx); ++Idx) { checkSlow(Idx < Elems.GetMaxIndex()); }
 	for (; Elems.IsValidIndex(Idx + Num); ++Num) {}
 	return { Data + NumBytes(Idx, Stride), static_cast<uint64>(Num) };
 }
@@ -1591,20 +1611,23 @@ inline void ReadSparseItems(FExistingItems& Dst, const FScriptSparseArray& Src,
 {
 	const uint8* Data = static_cast<const uint8*>(Src.GetData(0, Layout));
 
-	if (FExistingItemSlice LastRead = Dst.Slice)
+	if (Src.IsEmpty())
+	{
+		Dst.SetAll(nullptr, 0, Layout.Size);
+	}
+	else if (FExistingItemSlice LastRead = Dst.Slice)
 	{
 		// Continue partial response
-		int64 TotalBytesRead = static_cast<const uint8*>(LastRead.Data) - Data;
-		check(TotalBytesRead % Layout.Size == 0);
-		int32 NextIdx = /* skip known invalid */ 1 + TotalBytesRead / Layout.Size;
+		int64 PriorBytesRead = static_cast<const uint8*>(LastRead.Data) - Data;
+		check(PriorBytesRead % Layout.Size == 0);
+		int32 LastIdx = PriorBytesRead / Layout.Size;
+		int32 NextIdx = LastIdx + LastRead.Num + /* skip one known invalid */ 1;
 		check(NextIdx < Src.GetMaxIndex());
-		check(NextIdx > static_cast<int32>(LastRead.Num));
 		Dst.Slice = GetContiguousSlice(NextIdx, Src, Data, Layout.Size);
 	}
 	else if (Src.IsCompact())
 	{
-		uint64 Num = static_cast<uint64>(Src.Num());
-		Dst.SetAll(Num ? Data : nullptr, Num, Layout.Size);
+		Dst.SetAll(Data, static_cast<uint64>(Src.Num()), Layout.Size);
 	}
 	else
 	{
@@ -3078,7 +3101,7 @@ void BindStruct(FStructSchemaId Id, const UStruct* Struct)
 
 static void BindInitialTypes()
 {
-	UE_LOG(LogInit, Display, TEXT("Binding types to PlainProps schemas"));
+	UE_LOG(LogPlainPropsUObject, Display, TEXT("Binding types to PlainProps schemas..."));
 
 	// Declare all UScriptStruct/UClass/UFunction/UEnum
 	TArray<FStructSchemaId> Ids;
@@ -3307,9 +3330,9 @@ public:
 		SavedObjects.Reserve(NumReserve);
 	}
 
-	void Save(FStructSchemaId Id, const UObject* Object, const UObject* CDO)
+	void Save(FStructSchemaId Id, const UObject* Object, const UObject* Arch)
 	{
-		SavedObjects.Emplace(Id, CDO ? SaveStructDelta(Object, CDO, Id, DeltaCtx) : SaveStruct(Object, Id, FlatCtx));
+		SavedObjects.Emplace(Id, Arch ? SaveStructDelta(Object, Arch, Id, DeltaCtx) : SaveStruct(Object, Id, FlatCtx));
 	}
 
 	TArray64<uint8>	Write(TArray<FStructSchemaId>& OutRuntimeIds) const
@@ -3442,17 +3465,18 @@ private:
 
 struct FInstance
 {
-	UObject* Original;
+	UObject* Orig;
+	UObject* Arch;
 	UObject* PP;
 	UObject* TPS;
-	const UObject* CDO;
 	FStructSchemaId Id;
 
 	void Init()
 	{
-		const UClass* Class = Original->GetClass();
-		CDO = Class->GetDefaultObject(/* create lazily */ true);
-		CDO = CDO == Original ? nullptr : CDO;
+		UClass* Class = Orig->GetClass();
+		Class->GetDefaultObject(/* create lazily */ true);
+		Arch = Orig->GetArchetype();
+		Arch = Arch == Orig ? nullptr : Arch;
 		Id = GUE.Names.IndexStruct(IndexType(Class));
 	}
 };
@@ -3460,7 +3484,7 @@ struct FInstance
 static UObject* MakeEmptyInstance(UObject* Obj, FName Name)
 {
 	FStaticConstructObjectParameters Params(Obj->GetClass());
-	Params.Outer = Obj->GetOuter();
+	Params.Outer = Obj->GetOuter();//GetTransientPackage();
 	Params.Name = Name;
 	Params.SetFlags = Obj->GetFlags();
 	Params.Template = Obj->GetArchetype();
@@ -3469,43 +3493,57 @@ static UObject* MakeEmptyInstance(UObject* Obj, FName Name)
 	return StaticConstructObject_Internal(Params);
 }
 
-static bool Include(UClass* Class)
+static bool IncludeClass(UClass* Class)
+{
+	static const FName Exclusions[] = { "GameFeaturePluginStateMachine" };
+	static const FName SuperExclusions[] = { "LevelScriptActor" };
+
+	UStruct* Super = Class->GetInheritanceSuper();
+	return !Algo::Find(Exclusions, Class->GetFName()) && (!Super || !Algo::Find(SuperExclusions, Super->GetFName()));
+}
+
+static bool Include(EObjectFlags Flags, UClass* Class)
 {	
-	static const FName Exclusions[] = {"GameFeaturePluginStateMachine"};
-	return ShouldBind(Class) && !Algo::Find(Exclusions, Class->GetFName());
+	static constexpr EObjectFlags SkipFlags = RF_InheritableComponentTemplate;
+	return !(SkipFlags & Flags) && ShouldBind(Class) && IncludeClass(Class);
 }
-static void Roundtrip()
+
+static int32 Roundtrip()
 {
-	// Gather all non-empty objects
+	UE_LOG(LogPlainPropsUObject, Display, TEXT("Gathering all non-empty UObjects..."));
 	TArray<FInstance> Instances;
 	for (TObjectIterator<UObject> It; It; ++It)
 	{
 		UObject* Object = *It;
-		if (Include(Object->GetClass()))
+		if (Include(Object->GetFlags(), Object->GetClass()))
 		{
 			Instances.Add({*It});
 		}
 	}
 
 	// Create CDOs if needed and then clones for PP and TPS tests
+	UE_LOG(LogPlainPropsUObject, Display, TEXT("Cloning all UObjects..."));
 	for (FInstance& Instance : Instances)
 	{
 		Instance.Init();
 	}
+	FlushAsyncLoading();
+
 	FName PP_N("PP", 1);
 	for (FInstance& Instance : Instances)
 	{
 		PP_N.SetNumber(1 + &Instance - Instances.GetData());
-		Instance.PP = MakeEmptyInstance(Instance.Original, PP_N);
+		Instance.PP = MakeEmptyInstance(Instance.Orig, PP_N);
 	}
 	FName TPS_N("TPS", 1);
 	for (FInstance& Instance : Instances)
 	{
 		TPS_N.SetNumber(1 + &Instance - Instances.GetData());
-		Instance.TPS = MakeEmptyInstance(Instance.Original, TPS_N);
+		Instance.TPS = MakeEmptyInstance(Instance.Orig, TPS_N);
 	}
 
 	// Save
+	UE_LOG(LogPlainPropsUObject, Display, TEXT("Saving all UObjects to PlainProps..."));
 	FMemoryBatch Doc;
 	TCustomPropertyBindings<FMemoryPropertyBatch> FlatBindings(Doc.Properties, GUE.Customs);
 	TCustomPropertyBindings<FMemoryPropertyBatch> DeltaBindings(Doc.Properties, GUE.DeltaCustoms);
@@ -3513,12 +3551,13 @@ static void Roundtrip()
 		FBatchSaver Batch(FlatBindings.Overlay, DeltaBindings.Overlay, GUObjectArray.GetObjectArrayNum());
 		for (const FInstance& Instance : Instances)
 		{	
-			Batch.Save(Instance.Id, Instance.Original, Instance.CDO);
+			Batch.Save(Instance.Id, Instance.Orig, Instance.Arch ? Instance.PP : nullptr);
 		}
 		Doc.Data = Batch.Write(/* out */ Doc.RuntimeIds);
 	}
 
 	// Load
+	UE_LOG(LogPlainPropsUObject, Display, TEXT("Loading all UObjects from PlainProps..."));
 	FMemoryBatchLoader Batch(DeltaBindings.Overlay, MakeMemoryView(Doc.Data), Doc.RuntimeIds);
 	for (const FInstance& Instance : Instances)
 	{
@@ -3526,43 +3565,75 @@ static void Roundtrip()
 	}
 
 	// Diff original vs PlainProps
+	UE_LOG(LogPlainPropsUObject, Display, TEXT("Diffing all instances..."));
 	FDiffContext FlatCtx	= { GUE.Types, GUE.Schemas, FlatBindings.Overlay };
 	FDiffContext DeltaCtx	= { GUE.Types, GUE.Schemas, DeltaBindings.Overlay };
 	FDiffPath Path;
 	TUtf8StringBuilder<256> Diffs;
+	uint32 NumDiffs = 0;
 	for (FInstance& Instance : Instances)
 	{
-		FDiffContext& Ctx = Instance.CDO ? DeltaCtx : FlatCtx;
-		if (DiffStructs(Instance.Id, Instance.Original, Instance.PP, /* in-out */ Ctx))
+		FDiffContext& Ctx = Instance.Arch ? DeltaCtx : FlatCtx;
+		if (DiffStructs(Instance.Id, Instance.Orig, Instance.PP, /* in-out */ Ctx))
 		{
 			PrintDiff(/* out */ Diffs, Ctx.Out, GUE.Names);
 			//Batch.Reload(Instance.TPS, &Instance - Instances.GetData());
+			//UObject* Foo = MakeEmptyInstance(Instance.Orig, "Banana");
 			Ctx.Out.Reset();
 			Diffs.Append("\n");	
+			++NumDiffs;
 		}
 	}
 
-	UE_DEBUG_BREAK();
+	UE_LOG(LogPlainPropsUObject, Display, TEXT("Done - %d diffs detected"), NumDiffs);
+	return Diffs.Len();
 }
 
-static void TestBindings()
+static int32 TestBindings()
 {
-	if (FParse::Param(FCommandLine::Get(), TEXT("BindPlainProps")))
-	{
-		TScopedStructBinding<FTransform, FDefaultRuntime> Transform;
-		TScopedStructBinding<FGuid, FDefaultRuntime> Guid;
-		TScopedStructBinding<FColor, FDefaultRuntime> Color;
-		TScopedStructBinding<FLinearColor, FDefaultRuntime> LinearColor;
-		TScopedStructBinding<FFieldPath, FDefaultRuntime> FieldPath(GUE.Structs.FieldPath);
-		TScopedStructBinding<FScriptDelegate, FDefaultRuntime> Delegate(GUE.Structs.Delegate);
-		// MulticastDelegate declaration is shared with MulticastSparseDelegate
-		TScopedStructBinding<FMulticastScriptDelegate, FDefaultRuntime> InlineMulticast({GUE.Structs.MulticastInlineDelegate, GUE.Structs.MulticastDelegate});
-		InitBatchedProperties();
-		BindInitialTypes();
-		Roundtrip();
-	}
+	TScopedStructBinding<FTransform, FDefaultRuntime> Transform;
+	TScopedStructBinding<FGuid, FDefaultRuntime> Guid;
+	TScopedStructBinding<FColor, FDefaultRuntime> Color;
+	TScopedStructBinding<FLinearColor, FDefaultRuntime> LinearColor;
+	TScopedStructBinding<FFieldPath, FDefaultRuntime> FieldPath(GUE.Structs.FieldPath);
+	TScopedStructBinding<FScriptDelegate, FDefaultRuntime> Delegate(GUE.Structs.Delegate);
+	// MulticastDelegate declaration is shared with MulticastSparseDelegate
+	TScopedStructBinding<FMulticastScriptDelegate, FDefaultRuntime> InlineMulticast({GUE.Structs.MulticastInlineDelegate, GUE.Structs.MulticastDelegate});
+	InitBatchedProperties();
+	BindInitialTypes();
+	return Roundtrip();
 }
 
-static FDelegateHandle GTestPreMain = FCoreDelegates::OnCommandletPreMain.AddStatic(&TestBindings);
+} // namespace PlainProps::UE
+
+#include "PlainPropsCommandlets.h"
+#include "AssetRegistry/AssetData.h"
+#include "AssetRegistry/AssetRegistryModule.h"
+#include "Engine/World.h"
+
+UTestPlainPropsCommandlet::UTestPlainPropsCommandlet(const FObjectInitializer& Init)
+: Super(Init)
+{}
+
+int32 UTestPlainPropsCommandlet::Main(const FString& Params)
+{
+	//TArray<FString> MapNames;
+	//IFileManager::Get().FindFiles(Maps, TEXT("*") + FPackageName::GetMapPackageExtension(), true, false);
+	UE_LOG(LogPlainPropsUObject, Display, TEXT("Loading asset registry..."));
+	FAssetRegistryModule& AssetRegistryModule = FModuleManager::LoadModuleChecked<FAssetRegistryModule>(AssetRegistryConstants::ModuleName);
+	AssetRegistryModule.Get().SearchAllAssets(true);
+
+	UE_LOG(LogPlainPropsUObject, Display, TEXT("Gathering all maps..."));
+	TArray<FAssetData> Maps;
+	AssetRegistryModule.Get().GetAssetsByClass(UWorld::StaticClass()->GetClassPathName(), /* out */ Maps, true);
+
+	UE_LOG(LogPlainPropsUObject, Display, TEXT("Loading all maps..."));
+	for (FAssetData& Map : Maps)
+	{
+		Map.GetSoftObjectPath().LoadAsync({});
+	}
+	FlushAsyncLoading();
 
-} // namespace PlainProps::UE
\ No newline at end of file
+	UE_LOG(LogPlainPropsUObject, Display, TEXT("Starting test..."));
+	return PlainProps::UE::TestBindings();
+}
\ No newline at end of file
diff --git a/Engine/Plugins/Experimental/PlainPropsUObject/Source/Public/PlainPropsCommandlets.h b/Engine/Plugins/Experimental/PlainPropsUObject/Source/Public/PlainPropsCommandlets.h
new file mode 100644
--- /dev/null
+++ b/Engine/Plugins/Experimental/PlainPropsUObject/Source/Public/PlainPropsCommandlets.h
@@ -0,0 +1,15 @@
+// Copyright Epic Games, Inc. All Rights Reserved.
+
+#pragma once
+
+#include "CoreMinimal.h"
+#include "UObject/ObjectMacros.h"
+#include "Commandlets/Commandlet.h"
+#include "PlainPropsCommandlets.generated.h"
+
+UCLASS()
+class UTestPlainPropsCommandlet : public UCommandlet
+{
+	GENERATED_UCLASS_BODY()
+	virtual int32 Main(const FString& Params) override;
+};
