PlainProps initial stubbed out Verse BPVM bindings, bug fixes, improved error messages and more

#rnx

diff --git a/Engine/Plugins/Experimental/PlainProps/Source/Private/PlainPropsBind.cpp b/Engine/Plugins/Experimental/PlainProps/Source/Private/PlainPropsBind.cpp
--- a/Engine/Plugins/Experimental/PlainProps/Source/Private/PlainPropsBind.cpp
+++ b/Engine/Plugins/Experimental/PlainProps/Source/Private/PlainPropsBind.cpp
@@ -341,15 +341,19 @@ inline FTypeId IndexInnermost(FIdIndexerBase& Ids, FMemberBindType Type, FOption
 	return Ids.MakeLeafParameter(Leaf);
 }
 
-FTypeId	FMemberBinding::IndexParameterName(FIdIndexerBase& Ids) const
+FDualTypeId	FMemberBinding::IndexParameterName(FIdIndexerBase& Ids) const
 {
-	FTypeId Out = IndexInnermost(Ids, InnermostType, InnermostSchema);
+	FTypeId BindType = IndexInnermost(Ids, InnermostType, InnermostSchema);
+	FTypeId DeclType = BindType;
 	for (FRangeBinding Range : ReverseIterate(RangeBindings))
 	{
 		FTypeId SizeType = Ids.MakeRangeParameter(Range.GetSizeType());
-		Out = Ids.MakeAnonymousParametricType({ Out, SizeType });
+		FTypeId RangeBindName = {NoId, FTypenameId(Range.GetBindName())};
+		DeclType = Ids.MakeAnonymousParametricType({DeclType, SizeType});
+		BindType = Ids.MakeParametricType(RangeBindName, {BindType, SizeType});
 	}
-	return Out;
+	
+	return {BindType, DeclType};
 }
 
 FSchemaBindings::~FSchemaBindings()
diff --git a/Engine/Plugins/Experimental/PlainProps/Source/Private/PlainPropsBuildSchema.cpp b/Engine/Plugins/Experimental/PlainProps/Source/Private/PlainPropsBuildSchema.cpp
--- a/Engine/Plugins/Experimental/PlainProps/Source/Private/PlainPropsBuildSchema.cpp
+++ b/Engine/Plugins/Experimental/PlainProps/Source/Private/PlainPropsBuildSchema.cpp
@@ -67,6 +67,7 @@ struct FStructSchemaBuilder
 struct FEnumSchemaBuilder
 {
 	const FEnumDeclaration&						Declaration;
+	const FDebugIds&							Debug;
 	FEnumSchemaId								Id;
 	TOptional<ELeafWidth>						NotedWidth;
 	TSet<uint64>								NotedConstants;
@@ -112,7 +113,7 @@ FORCEINLINE FEnumSchemaBuilder&	FSchemasBuilder::NoteEnum(FEnumSchemaId Id)
 {
 	checkf(!bBuilt, TEXT("Noted new members after building"));
 	checkf(Id.Idx < uint32(DeclaredEnums.Num()) && DeclaredEnums[Id.Idx], TEXT("Undeclared enum '%s' noted"), *Debug.Print(Id));
-	return GetOrEmplace(EnumIndices, Id, Enums, *DeclaredEnums[Id.Idx], Id);
+	return GetOrEmplace(EnumIndices, Id, Enums, *DeclaredEnums[Id.Idx], Debug, Id);
 }
 
 FORCEINLINE FStructSchemaBuilder& FSchemasBuilder::NoteStruct(FStructSchemaId BindId)
@@ -402,8 +403,10 @@ FBuiltEnumSchema FEnumSchemaBuilder::Build() const
 		}
 	}
 	
-	check(	NotedConstants.Num() == Out.Constants.Num() || 
-			NotedConstants.Num() == Out.Constants.Num() + (Out.Mode == EEnumMode::Flag) );
+	checkf(	NotedConstants.Num() == Out.Constants.Num() || 
+			NotedConstants.Num() == Out.Constants.Num() + (Out.Mode == EEnumMode::Flag),
+			TEXT("Noted %d constants but included %d in fla%c enum %s"), 
+			NotedConstants.Num(), Out.Constants.Num(), "tg"[Out.Mode == EEnumMode::Flag], *Debug.Print(Id));
 	return Out;
 }
 
@@ -432,7 +435,7 @@ void FEnumSchemaBuilder::NoteValue(ELeafWidth Width, uint64 Value)
 			for (int32 Idx = NumValidated, Num = NotedConstants.Num(); Idx < Num; ++Idx)
 			{
 				uint64 Flag = NotedConstants.Get(FSetElementId::FromInteger(Idx));
-				checkf(Algo::FindBy(Declaration.GetEnumerators(), Flag, &FEnumerator::Constant), TEXT("Enum flag %d is undeclared"), Flag);
+				checkf(Algo::FindBy(Declaration.GetEnumerators(), Flag, &FEnumerator::Constant), TEXT("Enum flag %d is undeclared in %s"), Flag, *Debug.Print(Id));
 			}
 		}
 	}
@@ -442,7 +445,7 @@ void FEnumSchemaBuilder::NoteValue(ELeafWidth Width, uint64 Value)
 		NotedConstants.FindOrAdd(Value, /* out */ &bValidated);
 		if (!bValidated)
 		{
-			checkf(Algo::FindBy(Declaration.GetEnumerators(), Value, &FEnumerator::Constant), TEXT("Enum value %d is undeclared"), Value);
+			checkf(Algo::FindBy(Declaration.GetEnumerators(), Value, &FEnumerator::Constant), TEXT("Enum value %d is undeclared in %s"), Value, *Debug.Print(Id));
 		}
 	}
 }
diff --git a/Engine/Plugins/Experimental/PlainProps/Source/Private/PlainPropsDeclare.cpp b/Engine/Plugins/Experimental/PlainProps/Source/Private/PlainPropsDeclare.cpp
--- a/Engine/Plugins/Experimental/PlainProps/Source/Private/PlainPropsDeclare.cpp
+++ b/Engine/Plugins/Experimental/PlainProps/Source/Private/PlainPropsDeclare.cpp
@@ -11,7 +11,7 @@ namespace PlainProps
 //		 E.g. Saving E::All in E { A=1, B=2, All=A|B }, adding C=4, All=A|B|C and loading will load A|B
 //		 It's impossible to know if a user set A|B or All when saving though, we only have the value 3.
 //		 To really know, we'd need to instead save an enum oplog, i.e. {set A, set B} or {set All}
-static TConstArrayView<FEnumerator> StripAliases(TConstArrayView<FEnumerator> In, EEnumMode Mode, TArray<FEnumerator, TInlineAllocator<64>>& OutTmp)
+static TConstArrayView<FEnumerator> StripAliases(TArray<FEnumerator, TInlineAllocator<64>>& OutTmp, TConstArrayView<FEnumerator> In, EEnumMode Mode, const FDebugIds& Debug)
 {
 	TBitArray<> Aliases;
 	if (Mode == EEnumMode::Flag)
@@ -21,7 +21,7 @@ static TConstArrayView<FEnumerator> StripAliases(TConstArrayView<FEnumerator> In
 		for (FEnumerator E : In)
 		{
 			bool bAlias = E.Constant ? (Seen & E.Constant) == E.Constant : bSeen0;
-			checkf(bAlias || FMath::CountBits(E.Constant) <= 1, TEXT("Flag enums must use one bit per enumerator"));
+			checkf(bAlias || FMath::CountBits(E.Constant) <= 1, TEXT("Flag enums must use one bit per enumerator, %s is %llx"), *Debug.Print(E.Name), E.Constant);
 			
 			Aliases.Add(bAlias);
 			Seen |= E.Constant;
@@ -151,7 +151,7 @@ const FEnumDeclaration& FDeclarations::DeclareEnum(FEnumSchemaId Id, FTypeId Typ
 	TArray<FEnumerator, TInlineAllocator<64>> Tmp;
 	if (Policy == EEnumAliases::Strip)
 	{
-		Enumerators = StripAliases(Enumerators, Mode, /* out */ Tmp);
+		Enumerators = StripAliases(/* out */ Tmp, Enumerators, Mode, Debug);
 	}
 
 	FEnumDeclaration Header{Type, Mode, IntCastChecked<uint16>(Enumerators.Num())};
@@ -176,6 +176,12 @@ void FDeclarations::DropStructRef(FStructSchemaId DeclId)
 	}
 }
 
+const FStructDeclaration*
+FDeclarations::Find(FStructSchemaId Id) const
+{
+	return Id.Idx < uint32(DeclaredStructs.Num()) ? DeclaredStructs[Id.Idx].Get() : nullptr;
+}
+
 #if DO_CHECK
 void FDeclarations::Check(FEnumSchemaId Id) const
 {
diff --git a/Engine/Plugins/Experimental/PlainProps/Source/Private/PlainPropsPrint.cpp b/Engine/Plugins/Experimental/PlainProps/Source/Private/PlainPropsPrint.cpp
--- a/Engine/Plugins/Experimental/PlainProps/Source/Private/PlainPropsPrint.cpp
+++ b/Engine/Plugins/Experimental/PlainProps/Source/Private/PlainPropsPrint.cpp
@@ -1201,7 +1201,7 @@ void PrintDiff(FUtf8Builder& Out, const FDiffPath& Diff, const FDebugIds& Names)
 		}
 		else if (Node.Type.IsRange())
 		{
-			Out << Node.Meta.Range.GetId();	
+			Names.AppendString(Out, FTypenameId(Node.Meta.Range.GetBindName()));	
 		}
 		else if (FOptionalEnumSchemaId Enum = Node.Meta.Leaf)
 		{
diff --git a/Engine/Plugins/Experimental/PlainProps/Source/Private/PlainPropsSave.cpp b/Engine/Plugins/Experimental/PlainProps/Source/Private/PlainPropsSave.cpp
--- a/Engine/Plugins/Experimental/PlainProps/Source/Private/PlainPropsSave.cpp
+++ b/Engine/Plugins/Experimental/PlainProps/Source/Private/PlainPropsSave.cpp
@@ -199,16 +199,32 @@ FBuiltRange* SaveLeafRange(const void* Range, const ILeafRangeBinding& Binding,
 	return Allocator.GetAllocatedRange();
 }
 
-[[nodiscard]] static FBuiltRange* SaveStructRange(const void* Range, const IItemRangeBinding& ItemBinding, const FSaveContext& Ctx, FStructSchemaId Id)
+inline const FStructDeclaration& GetDeclaration(const FSaveContext& Ctx, FStructSchemaId BindId)
 {
-	if (Ctx.Defaults && Ctx.Declarations.Get(Id).Occupancy == EMemberPresence::AllowSparse)
+	// BindId is usually DeclId, unless it's been type-erased/lowered
+	if (const FStructDeclaration* Out = Ctx.Declarations.Find(BindId))
 	{
-		FDefaultStruct Default = { Id, Ctx.Defaults->Get(Id) };
+		return *Out;
+	}
+
+	if (const FStructDeclaration* Out = Ctx.Customs.FindDeclaration(BindId))
+	{
+		return *Out;
+	}
+	
+	return Ctx.Declarations.Get(Ctx.Schemas.GetStruct(BindId).DeclId);
+}
+
+[[nodiscard]] static FBuiltRange* SaveStructRange(const void* Range, const IItemRangeBinding& ItemBinding, const FSaveContext& Ctx, FStructSchemaId BindId)
+{
+	if (Ctx.Defaults && GetDeclaration(Ctx, BindId).Occupancy == EMemberPresence::AllowSparse)
+	{
+		FDefaultStruct Default = { BindId, Ctx.Defaults->Get(BindId) };
 		return SaveNonLeafRange<FInternalStructRangeDeltaSaver>(Range, ItemBinding, Ctx, Default);
 	}
 	else
 	{
-		return SaveNonLeafRange<FInternalStructRangeSaver>(Range, ItemBinding, Ctx, Id);
+		return SaveNonLeafRange<FInternalStructRangeSaver>(Range, ItemBinding, Ctx, BindId);
 	}
 }
 
diff --git a/Engine/Plugins/Experimental/PlainProps/Source/Public/PlainPropsBind.h b/Engine/Plugins/Experimental/PlainProps/Source/Public/PlainPropsBind.h
--- a/Engine/Plugins/Experimental/PlainProps/Source/Public/PlainPropsBind.h
+++ b/Engine/Plugins/Experimental/PlainProps/Source/Public/PlainPropsBind.h
@@ -488,10 +488,12 @@ struct FSaveRangeContext
 
 struct alignas(16) IItemRangeBinding
 {
-	// Try living without virtual destructor and trivially destructible members
 	virtual void ReadItems(FSaveRangeContext& Ctx) const = 0;
 	virtual void MakeItems(FLoadRangeContext& Ctx) const = 0;
-	virtual const char* GetId() const = 0;
+
+	// Try living without virtual destructor and trivially destructible members
+	explicit IItemRangeBinding(FConcreteTypenameId RangeBindName) : BindName(RangeBindName) {}
+	const FConcreteTypenameId BindName;
 };
 
 ////////////////////////////////////////////////////////////////////////////////////////////////
@@ -533,7 +535,10 @@ struct alignas(16) ILeafRangeBinding
 	virtual void		SaveLeaves(const void* Range, FLeafRangeAllocator& Out) const = 0;
 	virtual void		LoadLeaves(void* Range, FLeafRangeLoadView Leaves) const = 0;
 	virtual bool		DiffLeaves(const void* RangeA, const void* RangeB) const = 0;
-	virtual const char* GetId() const = 0;
+
+	explicit ILeafRangeBinding(FConcreteTypenameId RangeBindName) : BindName(RangeBindName) {}
+	const FConcreteTypenameId BindName;
+
 	template<typename SizeType>
 	static bool Diff(SizeType NumA, SizeType NumB, const void* A, const void* B, SIZE_T ItemSize)
 	{
@@ -554,11 +559,11 @@ public:
 	PLAINPROPS_API FRangeBinding(const IItemRangeBinding& Binding, ERangeSizeType SizeType);
 	PLAINPROPS_API FRangeBinding(const ILeafRangeBinding& Binding, ERangeSizeType SizeType);
 	
-	bool						IsLeafBinding() const		{ return !!(LeafMask & Handle); }
-	const IItemRangeBinding&	AsItemBinding() const		{ check(!IsLeafBinding()); return *reinterpret_cast<IItemRangeBinding*>(Handle & BindMask); }
-	const ILeafRangeBinding&	AsLeafBinding() const		{ check( IsLeafBinding()); return *reinterpret_cast<ILeafRangeBinding*>(Handle & BindMask); }
-	ERangeSizeType				GetSizeType() const			{ return static_cast<ERangeSizeType>(Handle & SizeMask); }
-	const char*					GetId() const				{ return IsLeafBinding() ? AsLeafBinding().GetId() : AsItemBinding().GetId(); }
+	bool						IsLeafBinding() const					{ return !!(LeafMask & Handle); }
+	const IItemRangeBinding&	AsItemBinding() const					{ check(!IsLeafBinding()); return *reinterpret_cast<IItemRangeBinding*>(Handle & BindMask); }
+	const ILeafRangeBinding&	AsLeafBinding() const					{ check( IsLeafBinding()); return *reinterpret_cast<ILeafRangeBinding*>(Handle & BindMask); }
+	ERangeSizeType				GetSizeType() const						{ return static_cast<ERangeSizeType>(Handle & SizeMask); }
+	FConcreteTypenameId			GetBindName() const						{ return IsLeafBinding() ? AsLeafBinding().BindName : AsItemBinding().BindName; }
 	inline bool	operator==(FRangeBinding O) const { return Handle == O.Handle; }
 };
 
@@ -570,6 +575,16 @@ using RangeBind = typename TRangeBind<T>::Type;
 
 ////////////////////////////////////////////////////////////////////////////////////////////////
 
+struct FDualTypeId
+{
+	FTypeId		BindType;
+	FTypeId		DeclType;
+
+	bool IsLowered() const { return BindType != DeclType; }
+};
+
+////////////////////////////////////////////////////////////////////////////////////////////////
+
 struct FMemberBinding
 {
 	explicit FMemberBinding(uint64 InOffset = 0)
@@ -582,7 +597,8 @@ struct FMemberBinding
 	FOptionalSchemaId				InnermostSchema;	// Enum or struct schema
 	TConstArrayView<FRangeBinding>	RangeBindings;		// Non-empty -> Range
 
-	PLAINPROPS_API FTypeId			IndexParameterName(FIdIndexerBase& Ids) const;
+	// @pre InnermostSchema isn't type-erased / lowered
+	PLAINPROPS_API FDualTypeId		IndexParameterName(FIdIndexerBase& Ids) const;
 };
 
 class FSchemaBindings : public IStructBindIds
@@ -608,7 +624,7 @@ private:
 // Lookup default instances when delta-saving struct ranges
 struct IDefaultStructs
 {
-	virtual const void* Get(FStructSchemaId Id) = 0;
+	virtual const void* Get(FStructSchemaId BindId) = 0;
 };
 
 ////////////////////////////////////////////////////////////////////////////////////////////////
@@ -728,7 +744,7 @@ FDualStructSchemaId GetStructDualId()
 template<class Ids, typename Ctti>
 FTypeId IndexCttiName()
 {
-	FTypenameId Name = Ids::IndexTypename(Ctti::Name);
+	FTypenameId Name{Ids::IndexTypename(Ctti::Name)};
 	FScopeId Namespace = NoId;
 	if constexpr (Ctti::Namespace[0] != '\0')
 	{
@@ -749,7 +765,7 @@ template<class Ids, Arithmetic T>
 FTypeId IndexArithmeticName()
 {
 	static constexpr FUnpackedLeafType Leaf = ReflectArithmetic<T>;
-	return { NoId, Ids::IndexTypename(ToAnsiView(ArithmeticName<Leaf.Type, Leaf.Width>)) };
+	return { NoId, FTypenameId(Ids::IndexTypename(ToAnsiView(ArithmeticName<Leaf.Type, Leaf.Width>))) };
 }
 
 template<class Ids, ETypename, Arithmetic Leaf>
@@ -785,7 +801,7 @@ FTypeId IndexParameterName()
 		else
 		{
 			using Typename = TTypename<T>;
-			FTypeId RangeBindName = { IndexNamespaceId<Ids, Typename>(), Ids::IndexTypename(ToAnsiView(Typename::RangeBindName)) };
+			FTypeId RangeBindName = { IndexNamespaceId<Ids, Typename>(), FTypenameId(Ids::IndexTypename(ToAnsiView(RangeBinding::BindName))) };
 			return Ids::GetIndexer().MakeParametricType(RangeBindName, {ItemParam, SizeParam});
 		}
 	}
@@ -883,7 +899,7 @@ struct TScopedStructBinding : FDualStructSchemaId, CustomBinding
 	}
 };
 
-// Scoped schema binding for native type
+// Specialization for schemabound structs
 template<class T, class Runtime>
 struct TScopedStructBinding<T, Runtime, void>
 {
@@ -979,21 +995,22 @@ struct TInnermostType<RangeBinding, 1>
 	using Type = typename RangeBinding::ItemType;
 };
 
-template<typename RangeBinding, uint32 N>
-TConstArrayView<FRangeBinding> GetRangeBindings()
+template<typename RangeBinding, typename Ids, uint32 N>
+TConstArrayView<FRangeBinding> GetRangeBindings() requires (N > 0)
 {
-	static_assert(N != 0);
-
 	struct FOnce
 	{
-		FOnce() : Binding(Instance, RangeSizeOf(typename RangeBinding::SizeType{})) {}
+		FOnce(FAnsiStringView RangeBindName)
+		: Instance(Ids::IndexTypename(RangeBindName))
+		, Binding(Instance, RangeSizeOf(typename RangeBinding::SizeType{}))
+		{}
 		RangeBinding Instance;
 		FRangeBinding Binding;
 	};
 
 	if constexpr (N == 1)
 	{
-		static FOnce Static;
+		static FOnce Static(ToAnsiView(RangeBinding::BindName));
 		return MakeArrayView(&Static.Binding, N);
 	}
 	else
@@ -1003,9 +1020,10 @@ TConstArrayView<FRangeBinding> GetRangeBindings()
 
 		struct FNestedOnce : FOnce
 		{
-			FNestedOnce() 
+			FNestedOnce()
+			: FOnce(ToAnsiView(InnerRangeBinding::BindName))
 			{
-				FMemory::Memcpy(NestedBindings, GetRangeBindings<InnerRangeBinding, N - 1>().GetData(), sizeof(NestedBindings));
+				FMemory::Memcpy(NestedBindings, GetRangeBindings<InnerRangeBinding, Ids,  N - 1>().GetData(), sizeof(NestedBindings));
 			}
 			
 			uint8 NestedBindings[sizeof(FRangeBinding) * (N - 1)] = {};
@@ -1045,7 +1063,7 @@ FMemberBinding BindMember(uint64 Offset)
 			constexpr uint32 NumRangeBindings = CountRangeBindings<RangeBinding>();
 			using InnermostType = typename TInnermostType<RangeBinding, NumRangeBindings>::Type;
 
-			Out.RangeBindings = GetRangeBindings<RangeBinding, NumRangeBindings>();
+			Out.RangeBindings = GetRangeBindings<RangeBinding, Ids, NumRangeBindings>();
 			Out.InnermostType = BindInnermostType<InnermostType, Ids>(Out.InnermostSchema);
 		}
 		else
diff --git a/Engine/Plugins/Experimental/PlainProps/Source/Public/PlainPropsDeclare.h b/Engine/Plugins/Experimental/PlainProps/Source/Public/PlainPropsDeclare.h
--- a/Engine/Plugins/Experimental/PlainProps/Source/Public/PlainPropsDeclare.h
+++ b/Engine/Plugins/Experimental/PlainProps/Source/Public/PlainPropsDeclare.h
@@ -72,6 +72,7 @@ public:
 
 	const FEnumDeclaration&							Get(FEnumSchemaId DeclId) const			{ Check(DeclId); return *DeclaredEnums[DeclId.Idx]; }
 	const FStructDeclaration&						Get(FStructSchemaId DeclId) const		{ Check(DeclId); return *DeclaredStructs[DeclId.Idx]; }
+	PLAINPROPS_API const FStructDeclaration*		Find(FStructSchemaId Id) const;
 	
 	TConstArrayView<TUniquePtr<FEnumDeclaration>>	GetEnums() const				{ return DeclaredEnums; }
 	TConstArrayView<TUniquePtr<FStructDeclaration>>	GetStructs() const				{ return DeclaredStructs; }
diff --git a/Engine/Plugins/Experimental/PlainProps/Source/Public/PlainPropsIndex.h b/Engine/Plugins/Experimental/PlainProps/Source/Public/PlainPropsIndex.h
--- a/Engine/Plugins/Experimental/PlainProps/Source/Public/PlainPropsIndex.h
+++ b/Engine/Plugins/Experimental/PlainProps/Source/Public/PlainPropsIndex.h
@@ -71,7 +71,8 @@ public:
 	virtual uint32						NumNames() const = 0;
 	uint32								NumEnums() const									{ return Enums.Num(); }
 	uint32								NumStructs() const									{ return Structs.Num(); }
-
+	
+	FConcreteTypenameId					IndexRangeBindName(const char* Name)				{ return InitParameterName(Name); }
 protected:
 	FNestedScopeIndexer					NestedScopes;
 	FParametricTypeIndexer				ParametricTypes;
diff --git a/Engine/Plugins/Experimental/PlainProps/Source/Public/PlainPropsSaveMember.h b/Engine/Plugins/Experimental/PlainProps/Source/Public/PlainPropsSaveMember.h
--- a/Engine/Plugins/Experimental/PlainProps/Source/Public/PlainPropsSaveMember.h
+++ b/Engine/Plugins/Experimental/PlainProps/Source/Public/PlainPropsSaveMember.h
@@ -107,7 +107,7 @@ struct TRangeSerializer
 
 	template<class Ids>
 	TRangeSerializer(TCustomInit<Ids> Init, TConstArrayView<FMemberId> Names)
-	: Bindings(GetRangeBindings<RangeBinding, NumRanges>().GetData())
+	: Bindings(GetRangeBindings<RangeBinding, Ids, NumRanges>().GetData())
 	{
 		for (uint16 Idx = 0; Idx < NumRanges; ++Idx)
 		{
diff --git a/Engine/Plugins/Experimental/PlainProps/Source/Public/PlainPropsUeCoreBindings.h b/Engine/Plugins/Experimental/PlainProps/Source/Public/PlainPropsUeCoreBindings.h
--- a/Engine/Plugins/Experimental/PlainProps/Source/Public/PlainPropsUeCoreBindings.h
+++ b/Engine/Plugins/Experimental/PlainProps/Source/Public/PlainPropsUeCoreBindings.h
@@ -14,6 +14,9 @@
 #include "PlainPropsStringUtil.h"
 #include "UObject/NameTypes.h"
 
+namespace Verse { class FNativeString; }
+using FVerseString = Verse::FNativeString;
+
 PP_REFLECT_STRUCT_TEMPLATE(, TTuple, void, Key, Value); // Todo handle TTuple and higher arities
 
 namespace UE::Math
@@ -120,6 +123,8 @@ struct TArrayBinding : IItemRangeBinding
 	using SizeType = int32;
 	using ItemType = T;
 	using ArrayType = TArray<T, Allocator>;
+	using IItemRangeBinding::IItemRangeBinding;
+	inline static constexpr std::string_view BindName = TTypename<ArrayType>::RangeBindName;
 
 	virtual void MakeItems(FLoadRangeContext& Ctx) const override
 	{
@@ -142,8 +147,6 @@ struct TArrayBinding : IItemRangeBinding
 		const ArrayType& Array = Ctx.Request.GetRange<ArrayType>();
 		Ctx.Items.SetAll(Array.GetData(), static_cast<uint64>(Array.Num()));
 	}
-
-	virtual const char* GetId() const override { return "Array"; }
 };
 
 //////////////////////////////////////////////////////////////////////////
@@ -154,7 +157,9 @@ struct TStringBinding : ILeafRangeBinding
 	using SizeType = int32;
 	using ItemType = char8_t;
 	using CharType = typename StringType::ElementType;
-
+	using ILeafRangeBinding::ILeafRangeBinding;
+	inline static constexpr std::string_view BindName = TTypename<StringType>::RangeBindName;
+	
 	virtual void SaveLeaves(const void* Range, FLeafRangeAllocator& Out) const override
 	{
 		const TArray<CharType>& Src = static_cast<const StringType*>(Range)->GetCharArray();
@@ -208,10 +213,8 @@ struct TStringBinding : ILeafRangeBinding
 		const StringType& A = *static_cast<const StringType*>(RangeA);
 		const StringType& B = *static_cast<const StringType*>(RangeB);
 		// Case-sensitive unnormalized comparison
-		return Diff(A.Len(), B.Len(), A.GetCharArray().GetData(), B.GetCharArray().GetData(), sizeof(TCHAR));
+		return Diff(A.Len(), B.Len(), GetData(A), GetData(B), sizeof(CharType));
 	}
-
-	virtual const char* GetId() const override { return "String"; }
 };
 
 //////////////////////////////////////////////////////////////////////////
@@ -221,6 +224,8 @@ struct TUniquePtrBinding : IItemRangeBinding
 {
 	using SizeType = bool;
 	using ItemType = T;
+	using IItemRangeBinding::IItemRangeBinding;
+	inline static constexpr std::string_view BindName = "UniquePtr";
 
 	virtual void MakeItems(FLoadRangeContext& Ctx) const override
 	{
@@ -253,8 +258,6 @@ struct TUniquePtrBinding : IItemRangeBinding
 		const TUniquePtr<T>& Ptr = Ctx.Request.GetRange<TUniquePtr<T>>();
 		Ctx.Items.SetAll(Ptr.Get(), Ptr ? 1 : 0);
 	}
-
-	virtual const char* GetId() const override { return "UniquePtr"; }
 };
 
 //////////////////////////////////////////////////////////////////////////
@@ -264,6 +267,8 @@ struct TOptionalBinding : IItemRangeBinding
 {
 	using SizeType = bool;
 	using ItemType = T;
+	using IItemRangeBinding::IItemRangeBinding;
+	inline static constexpr std::string_view BindName = "Optional";
 
 	virtual void MakeItems(FLoadRangeContext& Ctx) const override
 	{
@@ -298,8 +303,6 @@ struct TOptionalBinding : IItemRangeBinding
 		check(!Opt || reinterpret_cast<const T*>(&Opt) == &Opt.GetValue());
 		Ctx.Items.SetAll(reinterpret_cast<const T*>(Opt ? &Opt : nullptr), Opt ? 1 : 0);
 	}
-	
-	virtual const char* GetId() const override { return "Optional"; }
 };
 
 //////////////////////////////////////////////////////////////////////////
@@ -310,6 +313,8 @@ struct TSetBinding : IItemRangeBinding
 	using SizeType = int32;
 	using ItemType = T;
 	using SetType = TSet<T, KeyFuncs, SetAllocator>;
+	using IItemRangeBinding::IItemRangeBinding;
+	inline static constexpr std::string_view BindName = TTypename<SetType>::RangeBindName;
 
 	virtual void MakeItems(FLoadRangeContext& Ctx) const override
 	{
@@ -411,8 +416,6 @@ struct TSetBinding : IItemRangeBinding
 		for (; Elems.IsValidIndex(Idx + Num); ++Num) {}
 		return { &Elems[Idx], static_cast<uint64>(Num) };
 	}
-
-	virtual const char* GetId() const override { return "Set"; }
 };
 
 //////////////////////////////////////////////////////////////////////////
@@ -463,7 +466,12 @@ struct TPairBinding : ICustomBinding
 
 template <typename K, typename V, typename SetAllocator, typename KeyFuncs>
 struct TMapBinding : public TSetBinding<TPair<K, V>, KeyFuncs, SetAllocator>
-{};
+{
+	using MapType = TMap<K, V, SetAllocator, KeyFuncs>;
+	using Super = TSetBinding<TPair<K, V>, KeyFuncs, SetAllocator>;
+	using Super::Super;
+	inline static constexpr std::string_view BindName = TTypename<MapType>::RangeBindName;
+};
 
 //////////////////////////////////////////////////////////////////////////
 
@@ -936,6 +944,8 @@ template <> struct TTypename<FGuid>					{ inline static constexpr std::string_vi
 template <> struct TTypename<FColor>				{ inline static constexpr std::string_view DeclName = "Color"; };
 template <> struct TTypename<FLinearColor>			{ inline static constexpr std::string_view DeclName = "LinearColor"; };
 template <> struct TTypename<FString>				{ inline static constexpr std::string_view RangeBindName = "String"; };
+template <> struct TTypename<FUtf8String>			{ inline static constexpr std::string_view RangeBindName = "Utf8String"; };
+template <> struct TTypename<FAnsiString>			{ inline static constexpr std::string_view RangeBindName = "AnsiString"; };
 
 template <typename K, typename V>
 struct TTypename<TPair<K,V>>
@@ -952,7 +962,6 @@ template<typename T, typename Allocator>
 struct TTypename<TArray<T, Allocator>>
 {
 	inline static constexpr std::string_view RangeBindName = Concat<UeArrayName, ShortTypename<Allocator>>;
-	inline static constexpr std::string_view Namespace;
 };
 
 template<>
@@ -978,14 +987,12 @@ template <typename T, typename KeyFuncs, typename SetAllocator>
 struct TTypename<TSet<T, KeyFuncs, SetAllocator>>
 {
 	inline static constexpr std::string_view RangeBindName = Concat<UeSetName, ShortTypename<KeyFuncs>, ShortTypename<SetAllocator>>;
-	inline static constexpr std::string_view Namespace;
 };
 
 template <typename K, typename V, typename SetAllocator, typename KeyFuncs>
 struct TTypename<TMap<K, V, SetAllocator, KeyFuncs>>
 {
 	inline static constexpr std::string_view RangeBindName = Concat<UeMapName, ShortTypename<SetAllocator>, ShortTypename<KeyFuncs>>;
-	inline static constexpr std::string_view Namespace;
 };
 
 template<>
@@ -1000,6 +1007,7 @@ struct TRangeBind<TArray<T, Allocator>>
 template<> struct TRangeBind<FString> { using Type = UE::TStringBinding<FString>; };
 template<> struct TRangeBind<FAnsiString> {	using Type = UE::TStringBinding<FAnsiString>; };
 template<> struct TRangeBind<FUtf8String> {	using Type = UE::TStringBinding<FUtf8String>; };
+template<> struct TRangeBind<FVerseString> { using Type = UE::TStringBinding<FVerseString>; };
 
 template<typename T>
 struct TRangeBind<TUniquePtr<T>>
diff --git a/Engine/Plugins/Experimental/PlainProps/Source/Tests/PlainPropsUeCoreTests.cpp b/Engine/Plugins/Experimental/PlainProps/Source/Tests/PlainPropsUeCoreTests.cpp
--- a/Engine/Plugins/Experimental/PlainProps/Source/Tests/PlainPropsUeCoreTests.cpp
+++ b/Engine/Plugins/Experimental/PlainProps/Source/Tests/PlainPropsUeCoreTests.cpp
@@ -38,14 +38,14 @@ static FCustomBindingsOverlay	GDeltaCustoms(/* base*/ GCustoms);
 
 struct FIds
 {
-	static FNameId			IndexName(FAnsiStringView Name)				{ return GNames.MakeName(FName(Name)); }
-	static FMemberId		IndexMember(FAnsiStringView Name)			{ return GNames.NameMember(FName(Name)); }
-	static FTypenameId		IndexTypename(FAnsiStringView Name)			{ return GNames.MakeTypename(FName(Name)); }
-	static FScopeId			IndexScope(FAnsiStringView Name)			{ return GNames.MakeScope(FName(Name)); }
-	static FEnumSchemaId	IndexEnum(FTypeId Type)						{ return GNames.IndexEnum(Type); }
-	static FStructSchemaId	IndexStruct(FTypeId Type)					{ return GNames.IndexStruct(Type); }
-	static FIdIndexerBase&	GetIndexer()								{ return GNames; }
-	static const FDebugIds& GetDebug()									{ return GNames; }
+	static FNameId				IndexName(FAnsiStringView Name)		{ return GNames.MakeName(FName(Name)); }
+	static FMemberId			IndexMember(FAnsiStringView Name)	{ return GNames.NameMember(FName(Name)); }
+	static FConcreteTypenameId	IndexTypename(FAnsiStringView Name)	{ return GNames.NameType(FName(Name)); }
+	static FScopeId				IndexScope(FAnsiStringView Name)	{ return GNames.MakeScope(FName(Name)); }
+	static FEnumSchemaId		IndexEnum(FTypeId Type)				{ return GNames.IndexEnum(Type); }
+	static FStructSchemaId		IndexStruct(FTypeId Type)			{ return GNames.IndexStruct(Type); }
+	static FIdIndexerBase&		GetIndexer()						{ return GNames; }
+	static const FDebugIds&		GetDebug()							{ return GNames; }
 };
 
 ////////////////////////////////////////////////////////////////////////////////////////////////
diff --git a/Engine/Plugins/Experimental/PlainPropsUObject/Source/Private/PlainPropsUObjectBindings.cpp b/Engine/Plugins/Experimental/PlainPropsUObject/Source/Private/PlainPropsUObjectBindings.cpp
--- a/Engine/Plugins/Experimental/PlainPropsUObject/Source/Private/PlainPropsUObjectBindings.cpp
+++ b/Engine/Plugins/Experimental/PlainPropsUObject/Source/Private/PlainPropsUObjectBindings.cpp
@@ -10,6 +10,7 @@
 #include "Algo/Find.h"
 #include "Containers/PagedArray.h"
 #include "Hash/xxhash.h"
+#include "Logging/StructuredLog.h"
 #include "Math/PreciseFP.h"
 #include "Misc/CommandLine.h"
 #include "Misc/CoreDelegates.h"
@@ -27,11 +28,13 @@
 #include "UObject/UnrealType.h"
 #include "UObject/UObjectIterator.h"
 #include "UObject/VerseValueProperty.h"
+#include "UObject/VerseStringProperty.h"
 
 
 using FUnicastScriptDelegate = TScriptDelegate<FNotThreadSafeNotCheckedDelegateMode>;
 using FMulticastInvocationList = TArray<FUnicastScriptDelegate>;
 using FMulticastInvocationView = TConstArrayView<FUnicastScriptDelegate>;
+using FDelegateBase = TDelegateAccessHandlerBase<FNotThreadSafeDelegateMode>;
 
 // Temp hacks. Long-term either add FProperty getters for ctor/dtor/hash function pointers 
 // and delegate APIs for non-intrusive serialization or integrate PlainProps into Core/CoreUObject
@@ -47,6 +50,27 @@ UE_DEFINE_PRIVATE_MEMBER_PTR(FName, GUnicastDelegateFunctionName, FUnicastScript
 UE_DEFINE_PRIVATE_MEMBER_PTR(FMulticastInvocationList, GMulticastDelegateInvocationList, FMulticastScriptDelegate, InvocationList);
 UE_DEFINE_PRIVATE_MEMBER_PTR(bool, GSparseDelegateIsBound, FSparseDelegate, bIsBound);
 
+#if UE_DETECT_DELEGATES_RACE_CONDITIONS && 0
+UE_DEFINE_PRIVATE_MEMBER_PTR(FMRSWRecursiveAccessDetector, GDelegateAccessDetector, FDelegateBase, AccessDetector);
+struct FDelegateAccess : public FDelegateBase
+{
+	struct FReadScope : public FDelegateBase::FReadAccessScope
+	{
+		explicit FReadScope(const FDelegateBase& In) : FDelegateBase::FReadAccessScope(In.*GDelegateAccessDetector) {}
+	};
+	struct FWriteScope : public FDelegateBase::FWriteAccessScope
+	{
+		explicit FWriteScope(FDelegateBase& In) : FDelegateBase::FWriteAccessScope(In.*GDelegateAccessDetector) {}
+	};
+};
+#else
+struct FDelegateAccess
+{
+	struct FReadScope { FReadScope(const FDelegateBase& In) {} };
+	using FWriteScope = FReadScope;
+};
+#endif // UE_DETECT_DELEGATES_RACE_CONDITIONS
+
 
 DEFINE_LOG_CATEGORY_STATIC(LogPlainPropsUObject, Log, All);
 
@@ -240,6 +264,17 @@ FCommonTypenameIds::FCommonTypenameIds(TIdIndexer<FName>& Names)
 , Map(						Names.NameType("Map"))
 , Set(						Names.NameType("Set"))
 , Pair(						Names.NameType("Pair"))
+, LeafArray(				Names.NameType("LeafArray"))
+, TrivialArray(				Names.NameType("TrivialArray"))
+, NonTrivialArray(			Names.NameType("NonTrivialArray"))
+, StaticArray(				Names.NameType("StaticArray"))
+, TrivialOptional(			Names.NameType("TrivialOptional"))
+, IntrusiveOptional(		Names.NameType("IntrusiveOptional"))
+, NonIntrusiveOptional(		Names.NameType("NonIntrusiveOptional"))
+, String(					Names.NameType("String"))
+, Utf8String(				Names.NameType("Utf8String"))
+, AnsiString(				Names.NameType("AnsiString"))
+, VerseString(				Names.NameType("VerseString"))
 {}
 
 FCommonStructIds::FCommonStructIds(const FCommonScopeIds& Scopes, TIdIndexer<FName>& Names)
@@ -257,6 +292,9 @@ FCommonStructIds::FCommonStructIds(const FCommonScopeIds& Scopes, TIdIndexer<FNa
 , MulticastDelegate(		Names.IndexStruct({Scopes.CoreUObject, Names.MakeTypename("MulticastDelegate")}))
 , MulticastInlineDelegate(	Names.IndexStruct({Scopes.CoreUObject, Names.MakeTypename("MulticastInlineDelegate")}))
 , MulticastSparseDelegate(	Names.IndexStruct({Scopes.CoreUObject, Names.MakeTypename("MulticastSparseDelegate")}))
+, VerseFunction(			Names.IndexStruct({Scopes.CoreUObject, Names.MakeTypename("VerseFunction")}))
+, DynamicallyTypedValue(	Names.IndexStruct({Scopes.CoreUObject, Names.MakeTypename("DynamicallyTypedValue")}))
+, ReferencePropertyValue(	Names.IndexStruct({Scopes.CoreUObject, Names.MakeTypename("ReferencePropertyValue")}))
 {}
 
 FCommonMemberIds::FCommonMemberIds(TIdIndexer<FName>& Names)
@@ -293,7 +331,7 @@ static constexpr uint64 LeafMask =			CASTCLASS_FNumericProperty | CASTCLASS_FEnu
 static constexpr uint64 IntSMask =			CASTCLASS_FInt8Property | CASTCLASS_FInt16Property | CASTCLASS_FIntProperty | CASTCLASS_FInt64Property;
 static constexpr uint64 IntUMask =			CASTCLASS_FByteProperty | CASTCLASS_FUInt16Property | CASTCLASS_FUInt32Property | CASTCLASS_FUInt64Property;
 static constexpr uint64 ContainerMask =		CASTCLASS_FArrayProperty | CASTCLASS_FSetProperty| CASTCLASS_FMapProperty  | CASTCLASS_FOptionalProperty;
-static constexpr uint64 StringMask =		CASTCLASS_FStrProperty | CASTCLASS_FUtf8StrProperty | CASTCLASS_FAnsiStrProperty;
+static constexpr uint64 StringMask =		CASTCLASS_FStrProperty | CASTCLASS_FUtf8StrProperty | CASTCLASS_FAnsiStrProperty | CASTCLASS_FVerseStringProperty;
 static constexpr uint64 CommonStructMask =	CASTCLASS_FNameProperty | CASTCLASS_FTextProperty |  CASTCLASS_FFieldPathProperty |
 											CASTCLASS_FObjectProperty | CASTCLASS_FWeakObjectProperty | CASTCLASS_FSoftObjectProperty | CASTCLASS_FLazyObjectProperty |
 											CASTCLASS_FDelegateProperty | CASTCLASS_FMulticastInlineDelegateProperty;
@@ -514,6 +552,7 @@ void FDelegateBinding::Save(FMemberBuilder& Dst, const FScriptDelegate& Src, con
 {
 	if (Src.IsBound())
 	{
+		FDelegateAccess::FReadScope Scope(Src);
 		Dst.AddStruct(GUE.Members.Object, GUE.Structs.WeakObjectPtr, SaveStruct(&(Src.*GDelegateObject), GUE.Structs.WeakObjectPtr, Ctx));
 		Dst.AddStruct(GUE.Members.Function, GUE.Structs.Name, SaveStruct(&(Src.*GDelegateFunctionName), GUE.Structs.Name, Ctx));
 	}
@@ -524,6 +563,7 @@ void FDelegateBinding::Load(FScriptDelegate& Dst, FStructLoadView Src, ECustomLo
 	FMemberLoader Members(Src);
 	if (Members.HasMore())
 	{
+		FDelegateAccess::FWriteScope Scope(Dst);
 		LoadStruct(&(Dst.*GDelegateObject), Members.GrabStruct());
 		LoadStruct(&(Dst.*GDelegateFunctionName), Members.GrabStruct());
 	}
@@ -589,6 +629,7 @@ static FTypedRange SaveInvocations(const FMulticastInvocationList& In, const FSa
 
 static void SaveMulticastDelegate(FMemberBuilder& Dst, const FMulticastScriptDelegate& Src, const FSaveContext& Ctx)
 {
+	FDelegateAccess::FReadScope Scope(Src);
 	Dst.AddRange(GUE.Members.Invocations, SaveInvocations(Src.*GMulticastDelegateInvocationList, Ctx));
 }
 
@@ -608,6 +649,7 @@ static void LoadInvocations(FMulticastInvocationList& Dst, FStructRangeLoadView
 
 static void LoadMulticastDelegate(FMulticastScriptDelegate& Dst, FMemberLoader& Src)
 {
+	FDelegateAccess::FWriteScope Scope(Dst);
 	LoadInvocations(Dst.*GMulticastDelegateInvocationList, Src.GrabRange().AsStructs());
 }
 
@@ -761,6 +803,60 @@ static FStructSchemaId BindSparseDelegate(FStructSchemaId Owner, FMulticastSpars
 
 ////////////////////////////////////////////////////////////////////////////////////////////////
 
+void FVerseFunctionBinding::Save(FMemberBuilder& Dst, const FVerseFunction& Src, const FVerseFunction* Default, const FSaveContext& Ctx) const
+{
+	unimplemented();
+}
+
+void FVerseFunctionBinding::Load(FVerseFunction& Dst, FStructLoadView Src, ECustomLoadMethod Method) const
+{
+	unimplemented();
+}
+
+bool FVerseFunctionBinding::Diff(const FVerseFunction& A, const FVerseFunction& B)
+{
+	//return !(A == B);
+	return false;
+}
+	
+////////////////////////////////////////////////////////////////////////////////////////////////
+
+void FDynamicallyTypedValueBinding::Save(FMemberBuilder& Dst, const Type& Src, const Type* Default, const FSaveContext& Ctx) const
+{
+	unimplemented();
+}
+
+void FDynamicallyTypedValueBinding::Load(Type& Dst, FStructLoadView Src, ECustomLoadMethod Method) const
+{
+	unimplemented();
+}
+
+bool FDynamicallyTypedValueBinding::Diff(const ::UE::FDynamicallyTypedValue& A, const ::UE::FDynamicallyTypedValue& B)
+{
+	//return !UE::Verse::FRuntimeTypeDynamic::Get().AreIdentical(&A, &B);
+	return false;
+}
+
+////////////////////////////////////////////////////////////////////////////////////////////////
+
+void FReferencePropertyBinding::Save(FMemberBuilder& Dst, const FReferencePropertyValue& Src, const FReferencePropertyValue* Default, const FSaveContext& Ctx) const
+{
+	unimplemented();
+}
+
+void FReferencePropertyBinding::Load(FReferencePropertyValue& Dst, FStructLoadView Src, ECustomLoadMethod Method) const
+{
+	unimplemented();
+}
+
+bool FReferencePropertyBinding::Diff(const FReferencePropertyValue& A, const FReferencePropertyValue& B)
+{
+	unimplemented();
+	return false;
+}
+
+////////////////////////////////////////////////////////////////////////////////////////////////
+
 struct FInterfaceBinding final : ICustomBinding
 {
 	explicit FInterfaceBinding(UClass* Class) : InterfaceClass(Class) {}
@@ -866,7 +962,8 @@ struct FParameterBinding : FMemberBinding
 
 ////////////////////////////////////////////////////////////////////////////////////////////////
 
-static bool IsMax(FName EnumName, FName ValueName)
+template<const std::string_view& Suffix>
+static bool EndsWithDelimitedSuffix(FName EnumName, FName ValueName)
 {
 	if (ValueName.GetNumber() != NAME_NO_NUMBER_INTERNAL)
 	{
@@ -877,14 +974,20 @@ static bool IsMax(FName EnumName, FName ValueName)
 	ANSICHAR Buffer[NAME_SIZE];
 	ValueName.GetComparisonNameEntry()->GetAnsiName(Buffer);
 	FAnsiStringView Value(Buffer);
-	if (Value.Len() < 5 || !Value.EndsWith("MAX"))
+	if (Value.Len() >= Suffix.size() + 2 && Value.EndsWith(ToAnsiView(Suffix)))
 	{
-		return false;
+		// Todo: Check EnumName too, maybe based on ECppForm
+		char Delimiter = Value[Value.Len() - Suffix.size() - 1];
+		return Delimiter == ':' || Delimiter == '_';	
 	}
+	return false;
+}
 
-	// Todo: Check EnumName too, maybe based on ECppForm
-	char PreMaxChar = Value[Value.Len() - 4];
-	return PreMaxChar == ':' || '_';
+static bool DenyMaxValue(FName Enum)
+{
+	static const FName AllowsMax[] = {
+		"ESlateBrushMirrorType" };
+	return !Algo::Find(AllowsMax, Enum);
 }
 
 static FEnumSchemaId DeclareEnum(UEnum* Enum)
@@ -892,19 +995,29 @@ static FEnumSchemaId DeclareEnum(UEnum* Enum)
 	FTypeId Type = IndexType(Enum);
 	FEnumSchemaId Id = GUE.Names.IndexEnum(Type);
 	EEnumMode Mode = Enum->HasAnyEnumFlags(EEnumFlags::Flags) ? EEnumMode::Flag : EEnumMode::Flat;
+
+	// Skip _MAX and _All enumerators
+	FName EnumName = Enum->GetFName();
+	int32 Num = Enum->NumEnums();
+	if (Num > 0 && DenyMaxValue(EnumName))
+	{
+		static constexpr std::string_view Max = "MAX";
+		static constexpr std::string_view All = "All";
+		Num -= EndsWithDelimitedSuffix<Max>(EnumName, Enum->GetNameByIndex(Num - 1));
+		Num -= Num > 0 && Mode == EEnumMode::Flag && 
+				EndsWithDelimitedSuffix<All>(EnumName, Enum->GetNameByIndex(Num - 1));
+	}
+
 	TArray<FEnumerator, TInlineAllocator<64>> Enumerators;
-	for (int32 Idx = 0, Num = Enum->NumEnums(); Idx < Num; ++Idx)
+	for (int32 Idx = 0; Idx < Num; ++Idx)
 	{
 		FName ValueName = Enum->GetNameByIndex(Idx);
-		if (Idx < Num - 1 || !IsMax(Enum->GetFName(), ValueName))
-		{
-			Enumerators.Emplace(GUE.Names.MakeName(ValueName), static_cast<uint64>(Enum->GetValueByIndex(Idx)));
-		}
+		Enumerators.Emplace(GUE.Names.MakeName(ValueName), static_cast<uint64>(Enum->GetValueByIndex(Idx)));
 	}
-	
+
 #if WITH_METADATA
 	// IsMax() classifies more names as "max" than Enum->ContainsExistingMax()
-	check(Enumerators.Num() == Enum->NumEnums() || Enum->ContainsExistingMax() || Enum->HasMetaData(TEXT("Hidden"), Enumerators.Num()));
+	checkSlow(Enumerators.Num() == Num || Enum->ContainsExistingMax() || Enum->HasMetaData(TEXT("Hidden"), Num));
 #endif
 
 	GUE.Types.DeclareEnum(Id, Type, Mode, Enumerators, EEnumAliases::Strip);
@@ -1010,7 +1123,9 @@ struct TTrivialArrayBinding : IItemRangeBinding
 {
 	const FArrayPropertyInfo Info;
 
-	TTrivialArrayBinding(FArrayPropertyInfo InFo) : Info(InFo) {}
+	TTrivialArrayBinding(FArrayPropertyInfo InFo, FConcreteTypenameId BindName = GUE.Typenames.TrivialArray)
+	: IItemRangeBinding(BindName)
+	, Info(InFo) {}
 
 	virtual void ReadItems(FSaveRangeContext& Ctx) const override
 	{
@@ -1031,8 +1146,6 @@ struct TTrivialArrayBinding : IItemRangeBinding
 		
 		Ctx.Items.Set(Array.GetData(), static_cast<uint64>(NewNum), Info.ItemSize);
 	}
-
-	virtual const char* GetId() const override { return "TrivialArray"; }
 };
 
 // Currently can't extract constructor/destructor function pointers from FProperty, which
@@ -1045,7 +1158,7 @@ struct TNonTrivialArrayBinding : TTrivialArrayBinding<ScriptArray>
 	using TTrivialArrayBinding<ScriptArray>::Info;
 
 	TNonTrivialArrayBinding(FArrayPropertyInfo InFo, const FProperty* InNer)
-	: TTrivialArrayBinding<ScriptArray>(InFo)
+	: TTrivialArrayBinding<ScriptArray>(InFo, GUE.Typenames.NonTrivialArray)
 	, Inner(InNer)
 	{}
 
@@ -1076,14 +1189,13 @@ struct TNonTrivialArrayBinding : TTrivialArrayBinding<ScriptArray>
 			FMemory::Memzero(Items, Num * Info.ItemSize);
 		}
 	}
-	
-	virtual const char* GetId() const override { return "NonTrivialArray"; }
 };
 
 template<ELeafType Type, ELeafWidth Width>
 struct TLeafArrayBinding : ILeafRangeBinding
 {
 	inline static constexpr SIZE_T LeafSize = SizeOf(Width);
+	TLeafArrayBinding() : ILeafRangeBinding(GUE.Typenames.LeafArray) {}
 
 	virtual void SaveLeaves(const void* Range, FLeafRangeAllocator& Out) const override
 	{
@@ -1109,8 +1221,6 @@ struct TLeafArrayBinding : ILeafRangeBinding
 		return Diff(A.Num(), B.Num(), A.GetData(), B.GetData(), LeafSize);
 	}
 
-	virtual const char* GetId() const override { return "LeafArray"; }
-
 	inline constexpr uint64 NumBytes(int32 NumItems) const
 	{
 		return static_cast<uint64>(NumItems) * LeafSize;
@@ -1270,7 +1380,6 @@ inline auto CastAs(FLeafRangeLoadView In)
 
 inline EPropertyKind GetPropertyKind(FLeafBindType In)
 {
-	check(In.Bind.Type != ELeafBindType::BitfieldBool);
 	ELeafType Type = ToLeafType(In.Bind.Type);
 	ELeafWidth Width = In.Basic.Width;
 
@@ -1762,7 +1871,8 @@ struct TSetPropertyBinding : IItemRangeBinding, ICustomBinding
 	const TPropertyRangeSerializer<ElemKind>		Range;
 
 	TSetPropertyBinding(FSetProperty* In, FMemberBinding Elem)
-	: Layout(In->SetLayout)
+	: IItemRangeBinding(GUE.Typenames.Set)
+	, Layout(In->SetLayout)
 	, Inner(In->ElementProp)
 	, Range(Elem)
 	{
@@ -1846,8 +1956,6 @@ struct TSetPropertyBinding : IItemRangeBinding, ICustomBinding
 			Rehash(Set);
 		}
 	}
-
-	virtual const char* GetId() const override { return "Set"; }
 	
 	virtual void SaveCustom(FMemberBuilder& Dst, const void* InSrc, const void* InDefault, const FSaveContext& Ctx) override
 	{
@@ -2044,27 +2152,27 @@ class FSetBindings
 	TMap<FParameterBinding, FRangeBinding> RangeBindings;
 
 	template<EPropertyKind ItemKind>
-	IItemRangeBinding* BindNew(FSetProperty* Property, FMemberBinding Elem, const FStructDeclaration& Declaration) 
+	IItemRangeBinding* BindNew(FSetProperty* Property, FMemberBinding Elem, FStructSchemaId BindId, const FStructDeclaration& Declaration) 
 	{
 		// Todo: Ownership / memory leak
 		TSetPropertyBinding<ItemKind>* Out = new TSetPropertyBinding<ItemKind>(Property, Elem);
-		GUE.DeltaCustoms.BindStruct(Declaration.Id, *Out, Declaration, {});
+		GUE.DeltaCustoms.BindStruct(BindId, *Out, Declaration, {});
 		return Out;
 	}
 
-	IItemRangeBinding* DeltaBindNew(FSetProperty* Property, FMemberBinding Elem, const FStructDeclaration& Declaration)
+	IItemRangeBinding* DeltaBindNew(FSetProperty* Property, FMemberBinding Elem, FStructSchemaId BindId, const FStructDeclaration& Declaration)
 	{
 		switch (GetPropertyKind(Elem))
 		{
-		case EPropertyKind::Range:		return BindNew<EPropertyKind::Range>(Property, Elem, Declaration);
-		case EPropertyKind::Struct:		return BindNew<EPropertyKind::Struct>(Property, Elem, Declaration);
-		case EPropertyKind::Bool:		return BindNew<EPropertyKind::Bool>(Property, Elem, Declaration);
-		case EPropertyKind::U8:			return BindNew<EPropertyKind::U8>(Property, Elem, Declaration);
-		case EPropertyKind::U16:		return BindNew<EPropertyKind::U16>(Property, Elem, Declaration);
-		case EPropertyKind::U32:		return BindNew<EPropertyKind::U32>(Property, Elem, Declaration);
-		case EPropertyKind::U64:		return BindNew<EPropertyKind::U64>(Property, Elem, Declaration);
-		case EPropertyKind::F32:		return BindNew<EPropertyKind::F32>(Property, Elem, Declaration);
-		case EPropertyKind::F64:		return BindNew<EPropertyKind::F64>(Property, Elem, Declaration);
+		case EPropertyKind::Range:		return BindNew<EPropertyKind::Range	>(Property, Elem, BindId, Declaration);
+		case EPropertyKind::Struct:		return BindNew<EPropertyKind::Struct>(Property, Elem, BindId, Declaration);
+		case EPropertyKind::Bool:		return BindNew<EPropertyKind::Bool	>(Property, Elem, BindId, Declaration);
+		case EPropertyKind::U8:			return BindNew<EPropertyKind::U8	>(Property, Elem, BindId, Declaration);
+		case EPropertyKind::U16:		return BindNew<EPropertyKind::U16	>(Property, Elem, BindId, Declaration);
+		case EPropertyKind::U32:		return BindNew<EPropertyKind::U32	>(Property, Elem, BindId, Declaration);
+		case EPropertyKind::U64:		return BindNew<EPropertyKind::U64	>(Property, Elem, BindId, Declaration);
+		case EPropertyKind::F32:		return BindNew<EPropertyKind::F32	>(Property, Elem, BindId, Declaration);
+		case EPropertyKind::F64:		return BindNew<EPropertyKind::F64	>(Property, Elem, BindId, Declaration);
 		default: check(false);			return nullptr;
 		}
 	}
@@ -2079,14 +2187,18 @@ public:
 		}
 
 		// Index custom delta binding struct name
-		FTypeId Params[1] = { Elem.IndexParameterName(GUE.Names) };
-		FTypeId Type = { GUE.Scopes.Core, FTypenameId(GUE.Names.MakeParametricTypeId(GUE.Typenames.Set, Params)) };
-		FStructSchemaId Id = GUE.Names.IndexStruct(Type);
+		FDualTypeId Param = Elem.IndexParameterName(GUE.Names);
+		FTypeId BindType = FTypeId{ GUE.Scopes.Core, FTypenameId(GUE.Names.MakeParametricTypeId(GUE.Typenames.Set, MakeArrayView(&Param.BindType, 1))) };
+		FTypeId DeclType = Param.IsLowered()
+						 ? FTypeId{ GUE.Scopes.Core, FTypenameId(GUE.Names.MakeParametricTypeId(GUE.Typenames.Set, MakeArrayView(&Param.DeclType, 1))) }
+						 : BindType;
+		FStructSchemaId BindId = GUE.Names.IndexStruct(BindType);
+		FStructSchemaId DeclId = Param.IsLowered() ? GUE.Names.IndexStruct(DeclType) : BindId;
 		FMemberId Members[2] = { GUE.Members.Del, GUE.Members.Add };
 
 		// Todo: Ownership / memory leak
-		const FStructDeclaration& Declaration = GUE.Types.DeclareStruct(Id, Type, Members, EMemberPresence::AllowSparse);
-		IItemRangeBinding* Out = DeltaBindNew(Property, Elem, Declaration);
+		const FStructDeclaration& Declaration = GUE.Types.DeclareStruct(DeclId, DeclType, Members, EMemberPresence::AllowSparse);
+		IItemRangeBinding* Out = DeltaBindNew(Property, Elem, BindId, Declaration);
 		return RangeBindings.Emplace(Elem, FRangeBinding(*Out, DefaultRangeMax));
 	}
 };
@@ -2101,9 +2213,10 @@ struct TMapPropertyBinding : IItemRangeBinding
 	const FScriptMapLayout							Layout;
 	const TInnerProperty<KeyKind>					InnerKey;
 	const TInnerProperty<ValueKind>					InnerValue;
-
+	
 	TMapPropertyBinding(FMapProperty* In)
-	: Layout(In->MapLayout)
+	: IItemRangeBinding(GUE.Typenames.Map)
+	, Layout(In->MapLayout)
 	, InnerKey(In->KeyProp)
 	, InnerValue(In->ValueProp)
 	{}
@@ -2178,8 +2291,6 @@ struct TMapPropertyBinding : IItemRangeBinding
 			Rehash(Map);
 		}
 	}
-
-	virtual const char* GetId() const override { return "Map";}
 };
 
 ////////////////////////////////////////////////////////////////////////////////////////////////
@@ -2506,24 +2617,29 @@ class FPairBindings
 	inline FStructSchemaId BindImpl(FPair Pair)
 	{
 		check(Pair.KV[0].Offset == 0 && Pair.KV[1].Offset > 0);
-		if (const FStructSchemaId* Id = Bindings.Find(Pair))
+		if (const FStructSchemaId* BindId = Bindings.Find(Pair))
 		{
-			return *Id;
+			return *BindId;
 		}
 
-		// Index names
-		FTypeId Params[2] = { Pair.KV[0].IndexParameterName(GUE.Names), Pair.KV[1].IndexParameterName(GUE.Names) };
-		FTypeId Type = { GUE.Scopes.Core, FTypenameId(GUE.Names.MakeParametricTypeId(GUE.Typenames.Pair, Params)) };
-		FStructSchemaId Id = GUE.Names.IndexStruct(Type);
+		// Index names, can be optimized by checking if KeyParam / BindParam IsLowered()
+		FDualTypeId KeyParam = Pair.KV[0].IndexParameterName(GUE.Names);
+		FDualTypeId ValueParam = Pair.KV[1].IndexParameterName(GUE.Names);
+		FTypeId BindParams[2] = { KeyParam.BindType, ValueParam.BindType };
+		FTypeId DeclParams[2] = { KeyParam.DeclType, ValueParam.DeclType };
+		FTypeId BindType = { GUE.Scopes.Core, FTypenameId(GUE.Names.MakeParametricTypeId(GUE.Typenames.Pair, BindParams)) };
+		FTypeId DeclType = { GUE.Scopes.Core, FTypenameId(GUE.Names.MakeParametricTypeId(GUE.Typenames.Pair, DeclParams)) };
+		FStructSchemaId BindId = GUE.Names.IndexStruct(BindType);
+		FStructSchemaId DeclId = GUE.Names.IndexStruct(DeclType);
 		FMemberId Members[2] = { GUE.Members.Key, GUE.Members.Value };
 
 		// Todo: Ownership / memory leak
-		GUE.Types.DeclareStruct(Id, Type, Members, EMemberPresence::RequireAll);
-		GUE.Schemas.BindStruct(Id, Id, Pair.KV);
+		GUE.Types.DeclareStruct(DeclId, DeclType, Members, EMemberPresence::RequireAll);
+		GUE.Schemas.BindStruct(BindId, DeclId, Pair.KV);
 
-		Bindings.Emplace(Pair, Id);
+		Bindings.Emplace(Pair, BindId);
 
-		return Id;
+		return BindId;
 	}
 public:
 
@@ -2541,7 +2657,8 @@ static FPairBindings GPairs;
 
 inline void ReadBoolOptionalItem(FSaveRangeContext& Ctx, uint32 ItemSize)
 {
-	check((&Ctx.Request.GetRange<uint8>())[ItemSize] <= 1);
+	checkf((&Ctx.Request.GetRange<uint8>())[ItemSize] <= uint8(true),
+		TEXT("Non-intrusive TOptional::bIsSet should be true or false, but byte at offset %d was %d"), ItemSize, (&Ctx.Request.GetRange<uint8>())[ItemSize]);
 	bool bSet = (&Ctx.Request.GetRange<bool>())[ItemSize];
 	Ctx.Items.SetAll(bSet ? Ctx.Request.Range : nullptr, uint64(bSet), ItemSize);
 }
@@ -2554,21 +2671,23 @@ inline void MakeBoolOptionalItem(FLoadRangeContext& Ctx, uint32 ItemSize)
 	Ctx.Items.Set(&Ctx.Request.GetRange<uint8>(), uint64(bSet), ItemSize);
 }
 
+static constexpr std::string_view TrivialOptionalName = "TrivialOptional";
 template<uint32 ItemSize>
 struct TTrivialOptionalBinding : IItemRangeBinding
 {
+	TTrivialOptionalBinding() : IItemRangeBinding(GUE.Names.IndexRangeBindName(Concat<TrivialOptionalName, HexString<ItemSize>>.data())) {}
+
 	virtual void ReadItems(FSaveRangeContext& Ctx) const override { ReadBoolOptionalItem(Ctx, ItemSize); }
 	virtual void MakeItems(FLoadRangeContext& Ctx) const override { MakeBoolOptionalItem(Ctx, ItemSize); }
-	virtual const char* GetId() const override { return "TTrivialOptional"; }
 };
 
 struct FTrivialOptionalBinding : IItemRangeBinding
 {
 	const uint32 ItemSize;
-	explicit FTrivialOptionalBinding(uint32 Size) : ItemSize(Size) {}
+	explicit FTrivialOptionalBinding(uint32 Size) : IItemRangeBinding(GUE.Typenames.TrivialOptional), ItemSize(Size) {}
+
 	virtual void ReadItems(FSaveRangeContext& Ctx) const override { ReadBoolOptionalItem(Ctx, ItemSize); }
 	virtual void MakeItems(FLoadRangeContext& Ctx) const override { MakeBoolOptionalItem(Ctx, ItemSize); }
-	virtual const char* GetId() const override { return "FTrivialOptional"; }
 };
 
 struct FIntrusiveOptionalBinding : IItemRangeBinding
@@ -2576,7 +2695,11 @@ struct FIntrusiveOptionalBinding : IItemRangeBinding
 	const FProperty*	Inner;
 	const uint32		ItemSize;
 
-	explicit FIntrusiveOptionalBinding(const FProperty* In) : Inner(In), ItemSize(In->GetElementSize()) {}
+	explicit FIntrusiveOptionalBinding(const FProperty* In) 
+	: IItemRangeBinding(GUE.Typenames.IntrusiveOptional)
+	, Inner(In)
+	, ItemSize(In->GetElementSize())
+	{}
 	
 	virtual void ReadItems(FSaveRangeContext& Ctx) const override
 	{
@@ -2598,8 +2721,6 @@ struct FIntrusiveOptionalBinding : IItemRangeBinding
 			Ctx.Items.Set(nullptr, 0, ItemSize);
 		}
 	}
-
-	virtual const char* GetId() const override { return "IntrusiveOptional"; }
 };
 
 struct FNonIntrusiveOptionalBinding : IItemRangeBinding
@@ -2610,7 +2731,8 @@ struct FNonIntrusiveOptionalBinding : IItemRangeBinding
 	const bool			bDestructor;
 	
 	explicit FNonIntrusiveOptionalBinding(const FProperty* In)
-	: Inner(In)
+	: IItemRangeBinding(GUE.Typenames.NonIntrusiveOptional)
+	, Inner(In)
 	, ItemSize(In->GetElementSize())
 	, bConstructor(HasConstructor(In))
 	, bDestructor(HasDestructor(In))
@@ -2642,12 +2764,9 @@ struct FNonIntrusiveOptionalBinding : IItemRangeBinding
 			else
 			{
 				FMemory::Memzero(Value, ItemSize);
-			}
-			
+			}	
 		}
 	}
-
-	virtual const char* GetId() const override { return "NonIntrusiveOptional"; }
 };
 
 
@@ -2663,16 +2782,12 @@ class FOptionalBindings
 	TTrivialOptionalBinding<24> Trivial24;
 	TTrivialOptionalBinding<32> Trivial32;
 
-	alignas(PLATFORM_CACHE_LINE_SIZE)
-	TMap<FParameterBinding, FRangeBinding>	Bindings;
+	TMap<FParameterBinding, FRangeBinding>	NormalBindings;
+	TMap<FParameterBinding, FRangeBinding>	IntrusiveBindings;
 
 	const IItemRangeBinding* BindNew(FProperty* Inner)
 	{
-		if (Inner->HasIntrusiveUnsetOptionalState())
-		{
-			return new FIntrusiveOptionalBinding(Inner);
-		}
-		else if (HasConstructor(Inner) || HasDestructor(Inner))
+		if (HasConstructor(Inner) || HasDestructor(Inner))
 		{
 			return new FNonIntrusiveOptionalBinding(Inner);
 		}
@@ -2694,13 +2809,16 @@ public:
 	FRangeBinding Bind(FProperty* Inner, FMemberBinding Key)
 	{
 		check(Key.Offset == 0);
+		
+		bool bIntrusive = Inner->HasIntrusiveUnsetOptionalState();
+		TMap<FParameterBinding, FRangeBinding>& Bindings = bIntrusive ? IntrusiveBindings : NormalBindings;
 		if (const FRangeBinding* Binding = Bindings.Find(FParameterBinding(Key)))
 		{
 			return *Binding;
 		}
 
 		// Todo: Ownership / memory leak
-		const IItemRangeBinding* Out = BindNew(Inner);
+		const IItemRangeBinding* Out = bIntrusive ? new FIntrusiveOptionalBinding(Inner) : BindNew(Inner);
 		return Bindings.Emplace(Key, FRangeBinding(*Out, ERangeSizeType::Uni));
 	}
 };
@@ -2711,19 +2829,33 @@ static FOptionalBindings GOptionals;
 
 struct FStringBindings
 {
-	TStringBinding<FString>		TChrInstance;
-	TStringBinding<FAnsiString>	AnsiInstance;
-	TStringBinding<FUtf8String>	Utf8Instance;
-	FRangeBinding				TChr{TChrInstance, DefaultRangeMax};
-	FRangeBinding				Ansi{AnsiInstance, DefaultRangeMax};
-	FRangeBinding				Utf8{Utf8Instance, DefaultRangeMax};
+	TStringBinding<FString>			TCharInstance{GUE.Typenames.String};
+	TStringBinding<FUtf8String>		Utf8Instance{GUE.Typenames.Utf8String};
+	TStringBinding<FAnsiString>		AnsiInstance{GUE.Typenames.AnsiString};
+	TStringBinding<FUtf8String>		VerseInstance{GUE.Typenames.VerseString}; // Bypass Verse::FNativeString for now
+	FRangeBinding					TChar{TCharInstance, DefaultRangeMax};
+	FRangeBinding					Utf8{Utf8Instance, DefaultRangeMax};
+	FRangeBinding					Ansi{AnsiInstance, DefaultRangeMax};
+	FRangeBinding					Verse{VerseInstance, DefaultRangeMax};
 
-	FMemberBinding				Bind(FProperty* Property, uint64 CastFlags) const
+	inline const FRangeBinding&	SelectBinding(uint64 CastFlags) const
 	{
+		switch (CastFlags & StringMask)
+		{
+			case CASTCLASS_FStrProperty:			return TChar;
+			case CASTCLASS_FUtf8StrProperty:		return Utf8;
+			case CASTCLASS_FAnsiStrProperty:		return Ansi;
+			case CASTCLASS_FVerseStringProperty:	return Verse;
+			default:								break;
+		}
 		check(FMath::CountBits(CastFlags & StringMask) == 1);
-		bool bTChr = HasAny<CASTCLASS_FStrProperty>(CastFlags);
-		bool bAnsi = HasAny<CASTCLASS_FAnsiStrProperty>(CastFlags);
-		const FRangeBinding& Binding = bTChr ? TChr : (bAnsi ? Ansi : Utf8);
+		check(false);
+		return TChar;
+	}
+
+	FMemberBinding Bind(FProperty* Property, uint64 CastFlags) const
+	{
+		const FRangeBinding& Binding = SelectBinding(CastFlags);
 
 		FMemberBinding Out(Property->GetOffset_ForInternal());
 		Out.InnermostType = FMemberBindType(ReflectLeaf<char8_t>);
@@ -2731,7 +2863,7 @@ struct FStringBindings
 		return Out;
 	}
 };
-static const FStringBindings GStrings;
+static const FStringBindings GStrings; // static init dependency after GUE
 
 ////////////////////////////////////////////////////////////////////////////////////////////////
 
@@ -2740,19 +2872,14 @@ struct FStaticArrayBinding : IItemRangeBinding
 	uint32 Num;
 	uint32 Stride;
 
-	FStaticArrayBinding(uint32 InNum, uint32 InStride) : Num(InNum), Stride(InStride) {}
-
-	virtual void ReadItems(FSaveRangeContext& Ctx) const override
-	{
-		Ctx.Items.SetAll(Ctx.Request.Range, Num, Stride);
-	}
+	FStaticArrayBinding(uint32 InNum, uint32 InStride)
+	: IItemRangeBinding(GUE.Typenames.StaticArray)
+	, Num(InNum)
+	, Stride(InStride)
+	{}
 
-	virtual void MakeItems(FLoadRangeContext& Ctx) const override
-	{
-		Ctx.Items.Set(&Ctx.Request.GetRange<uint8>(), Num, Stride);
-	}
-	
-	virtual const char* GetId() const override { return "StaticArray"; }
+	virtual void				ReadItems(FSaveRangeContext& Ctx) const override	{ Ctx.Items.SetAll(Ctx.Request.Range, Num, Stride); }
+	virtual void				MakeItems(FLoadRangeContext& Ctx) const override	{ Ctx.Items.Set(&Ctx.Request.GetRange<uint8>(), Num, Stride); }
 };
 
 ////////////////////////////////////////////////////////////////////////////////////////////////
@@ -2772,6 +2899,10 @@ private:
 	TOptional<EMemberPresence>						OwnerOccupancy;
 	TPagedArray<FRangeBinding, 1024>				Ranges;
 	TArray<FMemberBinding, TInlineAllocator<64>>	Members;
+	// BPVM only?
+	const FName										VerseFunctionProperty{"VerseFunctionProperty"};
+	const FName										VerseDynamicProperty{"VerseDynamicProperty"};
+	const FName										ReferenceProperty{"ReferenceProperty"}; // Verse reference + FProperty*
 
 	TConstArrayView<FRangeBinding> AllocateRangeBindings(FRangeBinding Head, TConstArrayView<FRangeBinding> Tail)
 	{
@@ -2787,7 +2918,7 @@ private:
 		const int32 NewPages = Align(Ranges.Num() + OutNum, PageMax) / PageMax;
 		if (NewPages > Ranges.NumPages() && !Ranges.IsEmpty())
 		{
-			int32 NumPad = PageMax - Ranges.Num() % PageMax;
+			int32 NumPad = Align(Ranges.Num(), PageMax) - Ranges.Num();
 			Ranges.Append(Tail.Slice(0, NumPad));
 			check(Ranges.Num() % PageMax == 0);
 		}
@@ -2930,6 +3061,7 @@ private:
 
 	FMemberBinding BindSingleProperty(FProperty* Property)
 	{
+		FName PropertyTypename = Property->GetClass()->GetFName();
 		uint64 Flags = Property->GetCastFlags();
 		if (HasAny<LeafMask>(Flags))
 		{
@@ -2982,8 +3114,25 @@ private:
 				: BindVRestValue(static_cast<FVRestValueProperty*>(Property)); 
 		}
 	#endif
+		else
+		{
+			if (PropertyTypename == VerseFunctionProperty) // FVerseFunctionProperty
+			{
+				return BindAsStruct(Property, GUE.Structs.VerseFunction);
+			}
+			else if (PropertyTypename == VerseDynamicProperty) // FVerseDynamicProperty
+			{
+				return BindAsStruct(Property, GUE.Structs.DynamicallyTypedValue);
+			}
+			else if (PropertyTypename == ReferenceProperty) // FReferenceProperty
+			{
+				return BindAsStruct(Property, GUE.Structs.ReferencePropertyValue);
+			}
 
-		checkf(false, TEXT("Unrecognized cast flags %llx in property %s"), Flags, *Property->GetNameCPP());
+			
+		}
+
+		checkf(false, TEXT("Unrecognized class cast flags %llx in %s %s"), Flags, *PropertyTypename.ToString(), *Property->GetNameCPP());
 		return FMemberBinding(Property->GetOffset_ForInternal());
 	}
 
@@ -3101,7 +3250,7 @@ void BindStruct(FStructSchemaId Id, const UStruct* Struct)
 
 static void BindInitialTypes()
 {
-	UE_LOG(LogPlainPropsUObject, Display, TEXT("Binding types to PlainProps schemas..."));
+	UE_LOGFMT(LogPlainPropsUObject, Display, "Binding types to PlainProps schemas...");
 
 	// Declare all UScriptStruct/UClass/UFunction/UEnum
 	TArray<FStructSchemaId> Ids;
@@ -3330,7 +3479,7 @@ public:
 		SavedObjects.Reserve(NumReserve);
 	}
 
-	void Save(FStructSchemaId Id, const UObject* Object, const UObject* Arch)
+	FORCENOINLINE void Save(FStructSchemaId Id, const UObject* Object, const UObject* Arch)
 	{
 		SavedObjects.Emplace(Id, Arch ? SaveStructDelta(Object, Arch, Id, DeltaCtx) : SaveStruct(Object, Id, FlatCtx));
 	}
@@ -3433,18 +3582,17 @@ public:
 	{
 		check(LoadIdx == Objects.Num()); // Test should load all saved objects
 		Plans.Reset();
-		const FSchemaBatch* LoadSchemas = UnmountReadSchemas(Objects[0].Schema.Batch);
-		DestroyTranslatedSchemas(LoadSchemas);
+		UnmountReadSchemas(Objects[0].Schema.Batch);
 	}
 
-	void Load(UObject* Dst)
+	FORCENOINLINE void Load(UObject* Dst)
 	{
 		FStructView In = Objects[LoadIdx];
 		LoadStruct(Dst, In.Values, In.Schema.Id, *Plans);
 		++LoadIdx;
 	}
 
-	void Reload(UObject* Dst, int32 ReloadIdx)
+	FORCENOINLINE void Reload(UObject* Dst, int32 ReloadIdx)
 	{
 		FStructView In = Objects[ReloadIdx];
 		LoadStruct(Dst, In.Values, In.Schema.Id, *Plans);
@@ -3510,7 +3658,7 @@ static bool Include(EObjectFlags Flags, UClass* Class)
 
 static int32 Roundtrip()
 {
-	UE_LOG(LogPlainPropsUObject, Display, TEXT("Gathering all non-empty UObjects..."));
+	UE_LOGFMT(LogPlainPropsUObject, Display, "Gathering all non-empty UObjects...");
 	TArray<FInstance> Instances;
 	for (TObjectIterator<UObject> It; It; ++It)
 	{
@@ -3522,7 +3670,7 @@ static int32 Roundtrip()
 	}
 
 	// Create CDOs if needed and then clones for PP and TPS tests
-	UE_LOG(LogPlainPropsUObject, Display, TEXT("Cloning all UObjects..."));
+	UE_LOGFMT(LogPlainPropsUObject, Display, "Cloning all UObjects...");
 	for (FInstance& Instance : Instances)
 	{
 		Instance.Init();
@@ -3543,7 +3691,7 @@ static int32 Roundtrip()
 	}
 
 	// Save
-	UE_LOG(LogPlainPropsUObject, Display, TEXT("Saving all UObjects to PlainProps..."));
+	UE_LOGFMT(LogPlainPropsUObject, Display, "Saving all UObjects to PlainProps...");
 	FMemoryBatch Doc;
 	TCustomPropertyBindings<FMemoryPropertyBatch> FlatBindings(Doc.Properties, GUE.Customs);
 	TCustomPropertyBindings<FMemoryPropertyBatch> DeltaBindings(Doc.Properties, GUE.DeltaCustoms);
@@ -3557,7 +3705,7 @@ static int32 Roundtrip()
 	}
 
 	// Load
-	UE_LOG(LogPlainPropsUObject, Display, TEXT("Loading all UObjects from PlainProps..."));
+	UE_LOGFMT(LogPlainPropsUObject, Display, "Loading all UObjects from PlainProps...");
 	FMemoryBatchLoader Batch(DeltaBindings.Overlay, MakeMemoryView(Doc.Data), Doc.RuntimeIds);
 	for (const FInstance& Instance : Instances)
 	{
@@ -3565,7 +3713,7 @@ static int32 Roundtrip()
 	}
 
 	// Diff original vs PlainProps
-	UE_LOG(LogPlainPropsUObject, Display, TEXT("Diffing all instances..."));
+	UE_LOGFMT(LogPlainPropsUObject, Display, "Diffing all instances...");
 	FDiffContext FlatCtx	= { GUE.Types, GUE.Schemas, FlatBindings.Overlay };
 	FDiffContext DeltaCtx	= { GUE.Types, GUE.Schemas, DeltaBindings.Overlay };
 	FDiffPath Path;
@@ -3585,7 +3733,7 @@ static int32 Roundtrip()
 		}
 	}
 
-	UE_LOG(LogPlainPropsUObject, Display, TEXT("Done - %d diffs detected"), NumDiffs);
+	UE_LOGFMT(LogPlainPropsUObject, Display, "Done - {Diffs} diffs detected", NumDiffs);
 	return Diffs.Len();
 }
 
@@ -3599,6 +3747,12 @@ static int32 TestBindings()
 	TScopedStructBinding<FScriptDelegate, FDefaultRuntime> Delegate(GUE.Structs.Delegate);
 	// MulticastDelegate declaration is shared with MulticastSparseDelegate
 	TScopedStructBinding<FMulticastScriptDelegate, FDefaultRuntime> InlineMulticast({GUE.Structs.MulticastInlineDelegate, GUE.Structs.MulticastDelegate});
+	// Verse
+	TScopedStructBinding<FVerseFunction, FDefaultRuntime> VerseFunction(GUE.Structs.VerseFunction);
+	TScopedStructBinding<::UE::FDynamicallyTypedValue, FDefaultRuntime> DynamicallyTypedValue(GUE.Structs.DynamicallyTypedValue);
+	TScopedStructBinding<FReferencePropertyValue, FDefaultRuntime> ReferencePropertyValue(GUE.Structs.ReferencePropertyValue);
+
+	GUE.Defaults.BindZeroes(GUE.Structs.FieldPath, sizeof(FFieldPath), alignof(FFieldPath));
 	InitBatchedProperties();
 	BindInitialTypes();
 	return Roundtrip();
@@ -3617,23 +3771,32 @@ UTestPlainPropsCommandlet::UTestPlainPropsCommandlet(const FObjectInitializer& I
 
 int32 UTestPlainPropsCommandlet::Main(const FString& Params)
 {
-	//TArray<FString> MapNames;
-	//IFileManager::Get().FindFiles(Maps, TEXT("*") + FPackageName::GetMapPackageExtension(), true, false);
-	UE_LOG(LogPlainPropsUObject, Display, TEXT("Loading asset registry..."));
-	FAssetRegistryModule& AssetRegistryModule = FModuleManager::LoadModuleChecked<FAssetRegistryModule>(AssetRegistryConstants::ModuleName);
-	AssetRegistryModule.Get().SearchAllAssets(true);
+	if (int32 LoadIdx = Params.Find("-load="); LoadIdx != INDEX_NONE)
+	{
+		// todo: split and load Assets
+		const TCHAR* Assets = &Params[LoadIdx + 6];
+		UE_LOGFMT(LogPlainPropsUObject, Display, "Loading {Assets}...", Assets);
+		FSoftObjectPath("/BRRoot/BRRoot.BRRoot").LoadAsync({});
+	}
+	else // load all .umaps in asset registry
+	{
+		UE_LOGFMT(LogPlainPropsUObject, Display, "Loading asset registry...");
+		FAssetRegistryModule& AssetRegistryModule = FModuleManager::LoadModuleChecked<FAssetRegistryModule>(AssetRegistryConstants::ModuleName);
+		AssetRegistryModule.Get().SearchAllAssets(true);
 
-	UE_LOG(LogPlainPropsUObject, Display, TEXT("Gathering all maps..."));
-	TArray<FAssetData> Maps;
-	AssetRegistryModule.Get().GetAssetsByClass(UWorld::StaticClass()->GetClassPathName(), /* out */ Maps, true);
+		UE_LOGFMT(LogPlainPropsUObject, Display, "Gathering all maps...");
+		TArray<FAssetData> Maps;
+		AssetRegistryModule.Get().GetAssetsByClass(UWorld::StaticClass()->GetClassPathName(), /* out */ Maps, true);
 
-	UE_LOG(LogPlainPropsUObject, Display, TEXT("Loading all maps..."));
-	for (FAssetData& Map : Maps)
-	{
-		Map.GetSoftObjectPath().LoadAsync({});
+		UE_LOGFMT(LogPlainPropsUObject, Display, "Loading all maps...");
+		for (FAssetData& Map : Maps)
+		{
+			Map.GetSoftObjectPath().LoadAsync({});
+		}
 	}
+
 	FlushAsyncLoading();
 
-	UE_LOG(LogPlainPropsUObject, Display, TEXT("Starting test..."));
+	UE_LOGFMT(LogPlainPropsUObject, Display, "Starting test...");
 	return PlainProps::UE::TestBindings();
 }
\ No newline at end of file
diff --git a/Engine/Plugins/Experimental/PlainPropsUObject/Source/Public/PlainPropsUObjectRuntime.h b/Engine/Plugins/Experimental/PlainPropsUObject/Source/Public/PlainPropsUObjectRuntime.h
--- a/Engine/Plugins/Experimental/PlainPropsUObject/Source/Public/PlainPropsUObjectRuntime.h
+++ b/Engine/Plugins/Experimental/PlainPropsUObject/Source/Public/PlainPropsUObjectRuntime.h
@@ -7,10 +7,13 @@
 #include "PlainPropsIndex.h"
 #include "PlainPropsTypes.h"
 #include "Containers/Map.h"
+#include "UObject/DynamicallyTypedValue.h"
 #include "UObject/NameTypes.h"
 #include "UObject/ScriptDelegateFwd.h"
 
 struct FFieldPath;
+struct FReferencePropertyValue;
+struct FVerseFunction;
 
 namespace PlainProps::UE
 {
@@ -74,6 +77,17 @@ struct FCommonTypenameIds
 	FConcreteTypenameId			Map;
 	FConcreteTypenameId			Set;
 	FConcreteTypenameId			Pair;
+	FConcreteTypenameId			LeafArray;
+	FConcreteTypenameId			TrivialArray;
+	FConcreteTypenameId			NonTrivialArray;
+	FConcreteTypenameId			StaticArray;
+	FConcreteTypenameId			TrivialOptional;
+	FConcreteTypenameId			IntrusiveOptional;
+	FConcreteTypenameId			NonIntrusiveOptional;
+	FConcreteTypenameId			String;
+	FConcreteTypenameId			Utf8String;
+	FConcreteTypenameId			AnsiString;
+	FConcreteTypenameId			VerseString;
 };
 
 struct FCommonStructIds
@@ -94,6 +108,9 @@ struct FCommonStructIds
 	FStructSchemaId				MulticastDelegate;
 	FStructSchemaId				MulticastInlineDelegate;
 	FStructSchemaId				MulticastSparseDelegate;
+	FStructSchemaId				VerseFunction;
+	FStructSchemaId				DynamicallyTypedValue;
+	FStructSchemaId				ReferencePropertyValue;
 };
 
 struct FCommonMemberIds
@@ -200,6 +217,44 @@ struct FMulticastInlineDelegateBinding : ICustomBinding
 	static bool		Diff(const Type& A, const Type& B);
 };
 
+//////////////////////////////////////////////////////////////////////////
+
+struct FVerseFunctionBinding : ICustomBinding
+{
+	using Type =  FVerseFunction;
+	const FMemberId MemberIds[1];
+
+	FVerseFunctionBinding(TCustomInit<FIds>) : MemberIds{GUE.Members.Value} {}
+
+	void			Save(FMemberBuilder& Dst, const Type& Src, const Type* Default, const FSaveContext&) const;
+	void			Load(Type& Dst, FStructLoadView Src, ECustomLoadMethod Method) const;
+	static bool		Diff(const Type& A, const Type& B);
+};
+
+struct FDynamicallyTypedValueBinding : ICustomBinding
+{
+	using Type = ::UE::FDynamicallyTypedValue;
+	const FMemberId MemberIds[1];
+
+	FDynamicallyTypedValueBinding(TCustomInit<FIds>) : MemberIds{GUE.Members.Value} {}
+
+	void			Save(FMemberBuilder& Dst, const Type& Src, const Type* Default, const FSaveContext&) const;
+	void			Load(Type& Dst, FStructLoadView Src, ECustomLoadMethod Method) const;
+	static bool		Diff(const Type& A, const Type& B);
+};
+
+
+struct FReferencePropertyBinding : ICustomBinding
+{
+	using Type = FReferencePropertyValue;
+	const FMemberId MemberIds[1];
+
+	FReferencePropertyBinding(TCustomInit<FIds>) : MemberIds{GUE.Members.Value} {}
+
+	void			Save(FMemberBuilder& Dst, const Type& Src, const Type* Default, const FSaveContext&) const;
+	void			Load(Type& Dst, FStructLoadView Src, ECustomLoadMethod Method) const;
+	static bool		Diff(const Type& A, const Type& B);
+};
 
 } // namespace PlainProps::UE
 
@@ -208,4 +263,7 @@ namespace PlainProps
 template<> struct TCustomBind<FFieldPath> { using Type = UE::FFieldPathBinding; };
 template<> struct TCustomBind<FScriptDelegate> { using Type = UE::FDelegateBinding; };
 template<> struct TCustomBind<FMulticastScriptDelegate> { using Type = UE::FMulticastInlineDelegateBinding; };
+template<> struct TCustomBind<::UE::FDynamicallyTypedValue> { using Type = UE::FDynamicallyTypedValueBinding; };
+template<> struct TCustomBind<FReferencePropertyValue> { using Type = UE::FReferencePropertyBinding; };
+template<> struct TCustomBind<FVerseFunction> { using Type = UE::FVerseFunctionBinding; };
 }
\ No newline at end of file
