PlainProps added temp debug print functions and fixed a loading bug

#rnx

diff --git a/Engine/Plugins/Experimental/PlainProps/Source/Private/PlainPropsPrint.cpp b/Engine/Plugins/Experimental/PlainProps/Source/Private/PlainPropsPrint.cpp
--- a/Engine/Plugins/Experimental/PlainProps/Source/Private/PlainPropsPrint.cpp
+++ b/Engine/Plugins/Experimental/PlainProps/Source/Private/PlainPropsPrint.cpp
@@ -1158,6 +1158,25 @@ FString FDebugIds::Print(FOptionalMemberId Name) const
 	return FString(FStringView(StringCast<TCHAR>(Out.GetData(), Out.Len())));
 }
 
+FString FDebugIds::Print(FScopeId Scope) const
+{
+	TUtf8StringBuilder<128> Out;
+	AppendString(Out, Scope);
+	return FString(FStringView(StringCast<TCHAR>(Out.GetData(), Out.Len())));
+}
+
+FString FDebugIds::Print(FTypenameId Typename) const
+{
+	TUtf8StringBuilder<128> Out;
+	AppendString(Out, Typename);
+	return FString(FStringView(StringCast<TCHAR>(Out.GetData(), Out.Len())));
+}
+
+FString FDebugIds::Print(FConcreteTypenameId Typename) const
+{
+	return Print(FTypenameId(Typename));
+}
+
 FString FDebugIds::Print(FTypeId Type) const
 {
 	TUtf8StringBuilder<128> Out;
diff --git a/Engine/Plugins/Experimental/PlainProps/Source/Public/PlainPropsTypes.h b/Engine/Plugins/Experimental/PlainProps/Source/Public/PlainPropsTypes.h
--- a/Engine/Plugins/Experimental/PlainProps/Source/Public/PlainPropsTypes.h
+++ b/Engine/Plugins/Experimental/PlainProps/Source/Public/PlainPropsTypes.h
@@ -443,6 +443,9 @@ public:
 	PLAINPROPS_API FString				Print(FNameId Name) const;
 	PLAINPROPS_API FString				Print(FMemberId Name) const;
 	PLAINPROPS_API FString				Print(FOptionalMemberId Name) const;
+	PLAINPROPS_API FString				Print(FScopeId Scope) const;
+	PLAINPROPS_API FString				Print(FTypenameId Typename) const;
+	PLAINPROPS_API FString				Print(FConcreteTypenameId Typename) const;
 	PLAINPROPS_API FString				Print(FTypeId Type) const;
 	PLAINPROPS_API FString				Print(FEnumSchemaId Name) const;
 	PLAINPROPS_API FString				Print(FStructSchemaId Name) const;
diff --git a/Engine/Plugins/Experimental/PlainPropsUObject/Source/Private/PlainPropsUObjectBindings.cpp b/Engine/Plugins/Experimental/PlainPropsUObject/Source/Private/PlainPropsUObjectBindings.cpp
--- a/Engine/Plugins/Experimental/PlainPropsUObject/Source/Private/PlainPropsUObjectBindings.cpp
+++ b/Engine/Plugins/Experimental/PlainPropsUObject/Source/Private/PlainPropsUObjectBindings.cpp
@@ -1162,8 +1162,6 @@ struct TNonTrivialArrayBinding : TTrivialArrayBinding<ScriptArray>
 	, Inner(InNer)
 	{}
 
-	//todo: other range bindings
-
 	virtual void MakeItems(FLoadRangeContext& Ctx) const override
 	{
 		ScriptArray& Array = Ctx.Request.GetRange<ScriptArray>();
@@ -2690,15 +2688,38 @@ struct FTrivialOptionalBinding : IItemRangeBinding
 	virtual void MakeItems(FLoadRangeContext& Ctx) const override { MakeBoolOptionalItem(Ctx, ItemSize); }
 };
 
-struct FIntrusiveOptionalBinding : IItemRangeBinding
+struct FOptionalBindingBase : IItemRangeBinding
 {
 	const FProperty*	Inner;
 	const uint32		ItemSize;
-
-	explicit FIntrusiveOptionalBinding(const FProperty* In) 
-	: IItemRangeBinding(GUE.Typenames.IntrusiveOptional)
+	const bool			bConstructor;
+	const bool			bDestructor;
+	
+	explicit FOptionalBindingBase(const FProperty* In, FConcreteTypenameId BindName)
+	: IItemRangeBinding(BindName)
 	, Inner(In)
 	, ItemSize(In->GetElementSize())
+	, bConstructor(HasConstructor(In))
+	, bDestructor(HasDestructor(In))
+	{}
+
+	void InitItem(void* Value) const
+	{
+		if (bConstructor)
+		{
+			ConstructValue(Inner, Value);
+		}
+		else
+		{
+			FMemory::Memzero(Value, ItemSize);
+		}
+	}
+};
+
+struct FIntrusiveOptionalBinding : FOptionalBindingBase
+{
+	explicit FIntrusiveOptionalBinding(const FProperty* In)
+	: FOptionalBindingBase(In, GUE.Typenames.IntrusiveOptional)
 	{}
 	
 	virtual void ReadItems(FSaveRangeContext& Ctx) const override
@@ -2711,9 +2732,10 @@ struct FIntrusiveOptionalBinding : IItemRangeBinding
 	{
 		uint8* Value = &Ctx.Request.GetRange<uint8>();
 		Inner->ClearIntrusiveOptionalValue(Value);
+
 		if (Ctx.Request.NumTotal() > 0)
 		{
-			Inner->InitializeIntrusiveUnsetOptionalValue(Value);
+			InitItem(Value);
 			Ctx.Items.Set(Value, 1, ItemSize);
 		}
 		else
@@ -2723,19 +2745,10 @@ struct FIntrusiveOptionalBinding : IItemRangeBinding
 	}
 };
 
-struct FNonIntrusiveOptionalBinding : IItemRangeBinding
+struct FNonIntrusiveOptionalBinding : FOptionalBindingBase
 {
-	const FProperty*	Inner;
-	const uint32		ItemSize;
-	const bool			bConstructor;
-	const bool			bDestructor;
-	
 	explicit FNonIntrusiveOptionalBinding(const FProperty* In)
-	: IItemRangeBinding(GUE.Typenames.NonIntrusiveOptional)
-	, Inner(In)
-	, ItemSize(In->GetElementSize())
-	, bConstructor(HasConstructor(In))
-	, bDestructor(HasDestructor(In))
+	: FOptionalBindingBase(In, GUE.Typenames.NonIntrusiveOptional)
 	{}
 
 	virtual void ReadItems(FSaveRangeContext& Ctx) const override
@@ -2754,23 +2767,13 @@ struct FNonIntrusiveOptionalBinding : IItemRangeBinding
 
 		bSet = Ctx.Request.NumTotal() > 0;
 		Ctx.Items.Set(bSet ? Value : nullptr , 1, ItemSize);
-
 		if (bSet)
 		{
-			if (bConstructor)
-			{
-				ConstructValue(Inner, Value);
-			}
-			else
-			{
-				FMemory::Memzero(Value, ItemSize);
-			}	
+			InitItem(Value);	
 		}
 	}
 };
 
-
-
 class FOptionalBindings
 {
 	TTrivialOptionalBinding<1> Trivial1;
@@ -3733,7 +3736,7 @@ static int32 Roundtrip()
 		}
 	}
 
-	UE_LOGFMT(LogPlainPropsUObject, Display, "Done - {Diffs} diffs detected", NumDiffs);
+	UE_LOGFMT(LogPlainPropsUObject, Display, "Done - {Diffs} diffs in {Objs} UObjects saved in a {KB}KB value stream", NumDiffs, Instances.Num(), Doc.Data.NumBytes() / 1024);
 	return Diffs.Len();
 }
 
@@ -3799,4 +3802,15 @@ int32 UTestPlainPropsCommandlet::Main(const FString& Params)
 
 	UE_LOGFMT(LogPlainPropsUObject, Display, "Starting test...");
 	return PlainProps::UE::TestBindings();
-}
\ No newline at end of file
+}
+
+// Stopgap debugger callable functions til natvis is added 
+extern FString _print(const PlainProps::FNameId& Id)				{ return PlainProps::UE::GUE.Names.Print(Id); }
+extern FString _print(const PlainProps::FMemberId& Id)			{ return PlainProps::UE::GUE.Names.Print(Id); }
+extern FString _print(const PlainProps::FOptionalMemberId& Id)	{ return PlainProps::UE::GUE.Names.Print(Id); }
+extern FString _print(const PlainProps::FScopeId& Id)				{ return PlainProps::UE::GUE.Names.Print(Id); }
+extern FString _print(const PlainProps::FTypenameId& Id)			{ return PlainProps::UE::GUE.Names.Print(Id); }
+extern FString _print(const PlainProps::FConcreteTypenameId& Id)	{ return PlainProps::UE::GUE.Names.Print(Id); }
+extern FString _print(const PlainProps::FTypeId& Id)				{ return PlainProps::UE::GUE.Names.Print(Id); }
+extern FString _print(const PlainProps::FEnumSchemaId& Id)		{ return PlainProps::UE::GUE.Names.Print(Id); }
+extern FString _print(const PlainProps::FStructSchemaId& Id)		{ return PlainProps::UE::GUE.Names.Print(Id); }
\ No newline at end of file
