PlainProps roundtrip all CitySample .umaps cleanly

#rnx

diff --git a/Engine/Plugins/Experimental/PlainPropsUObject/Source/Private/PlainPropsUObjectBindings.cpp b/Engine/Plugins/Experimental/PlainPropsUObject/Source/Private/PlainPropsUObjectBindings.cpp
--- a/Engine/Plugins/Experimental/PlainPropsUObject/Source/Private/PlainPropsUObjectBindings.cpp
+++ b/Engine/Plugins/Experimental/PlainPropsUObject/Source/Private/PlainPropsUObjectBindings.cpp
@@ -3481,6 +3481,7 @@ struct FMemoryBatch
 };
 
 static constexpr uint32 Magics[] = { 0xFEEDF00D, 0xABCD1234, 0xDADADAAA, 0x99887766, 0xF0F1F2F3 };
+volatile static const UObject* GDebugNoteObject;
 
 class FBatchSaver
 {
@@ -3494,7 +3495,8 @@ public:
 
 	FORCENOINLINE void Save(FStructSchemaId Id, const UObject* Object, const UObject* Arch)
 	{
-		SavedObjects.Emplace(Id, Arch ? SaveStructDelta(Object, Arch, Id, DeltaCtx) : SaveStruct(Object, Id, FlatCtx));
+		FBuiltStruct* Built = Arch ? SaveStructDelta(Object, Arch, Id, DeltaCtx) : SaveStruct(Object, Id, FlatCtx);
+		SavedObjects.Emplace(Id, Built, Object);
 	}
 
 	TArray64<uint8>	Write(TArray<FStructSchemaId>& OutRuntimeIds) const
@@ -3502,12 +3504,12 @@ public:
 		// Build partial schemas
 		const FStructBindIds BindIds(GUE.DeltaCustoms, GUE.Schemas);
 		FSchemasBuilder SchemaBuilders(GUE.Types, BindIds, Scratch, ESchemaFormat::InMemoryNames);
-		uint32 DebugIdx = 0;
-		for (const IdBuiltStructPair& Object : SavedObjects)
+		for (const FSavedObject& Object : SavedObjects)
 		{
-			SchemaBuilders.NoteStructAndMembers(Object.Key, *Object.Value);
-			++DebugIdx;
+			GDebugNoteObject = Object.Input;
+			SchemaBuilders.NoteStructAndMembers(Object.Id, *Object.Built);
 		}
+		GDebugNoteObject = nullptr;
 		FBuiltSchemas Schemas = SchemaBuilders.Build();
 
 		// Save schema ids on the side
@@ -3525,11 +3527,11 @@ public:
 
 		// Write objects
 		WriteInt(Out, Magics[2]);
-		for (const TPair<FStructSchemaId, FBuiltStruct*>& Object : SavedObjects)
+		for (const FSavedObject& Object : SavedObjects)
 		{
 			WriteInt(/* out */ Tmp, Magics[3]);
-			WriteInt(/* out */ Tmp, Writer.GetWriteId(Object.Key).Get().Idx);
-			Writer.WriteMembers(/* out */ Tmp, Object.Key, *Object.Value);
+			WriteInt(/* out */ Tmp, Writer.GetWriteId(Object.Id).Get().Idx);
+			Writer.WriteMembers(/* out */ Tmp, Object.Id, *Object.Built);
 			WriteSkippableSlice(Out, Tmp);
 			Tmp.Reset();
 		}
@@ -3542,9 +3544,14 @@ public:
 	}
 
 private:
-	using IdBuiltStructPair = TPair<FStructSchemaId, FBuiltStruct*>;
+	struct FSavedObject
+	{
+		FStructSchemaId		Id;
+		FBuiltStruct*		Built;
+		const UObject*		Input; // For debug
+	};
 
-	TArray<IdBuiltStructPair>	SavedObjects;
+	TArray<FSavedObject>		SavedObjects;
 	mutable FScratchAllocator	Scratch;
 	FSaveContext				FlatCtx;
 	FSaveContext				DeltaCtx;
@@ -3658,7 +3665,11 @@ static UObject* MakeEmptyInstance(UObject* Obj, FName Name)
 
 static bool IncludeClass(UClass* Class)
 {
-	static const FName Exclusions[] = { "GameFeaturePluginStateMachine", "WorldSettings" };
+	static const FName Exclusions[] = {
+		"GameFeaturePluginStateMachine", // Cloning ensure
+		"WorldSettings", // QAGame
+		"AnimationSequencerDataModel" // CitySample - Enum value 8 is undeclared in /Script/Engine.ERichCurveTangentMode, illegal value detected in /Script/Engine.RichCurveKey::TangentMode 
+	};
 	static const FName SuperExclusions[] = { "LevelScriptActor" };
 
 	if (!ShouldBind(Class) || !!Algo::Find(Exclusions, Class->GetFName()))
@@ -3821,7 +3832,7 @@ int32 UTestPlainPropsCommandlet::Main(const FString& Params)
 		TArray<FAssetData> Maps;
 		AssetRegistryModule.Get().GetAssetsByClass(UWorld::StaticClass()->GetClassPathName(), /* out */ Maps, true);
 
-		UE_LOGFMT(LogPlainPropsUObject, Display, "Loading all maps...");
+		UE_LOGFMT(LogPlainPropsUObject, Display, "Loading all {Maps} maps...", Maps.Num());
 		for (FAssetData& Map : Maps)
 		{
 			Map.GetSoftObjectPath().LoadAsync({});
