PlainProps Add TPS comparison, fix FUnpackedLeafBindType bug, and change to PreciseFP equality

#rnx

diff --git a/Engine/Plugins/Experimental/PlainProps/Source/PlainProps.Build.cs b/Engine/Plugins/Experimental/PlainProps/Source/PlainProps.Build.cs
--- a/Engine/Plugins/Experimental/PlainProps/Source/PlainProps.Build.cs
+++ b/Engine/Plugins/Experimental/PlainProps/Source/PlainProps.Build.cs
@@ -13,6 +13,7 @@ public class PlainProps : ModuleRules
 		PrivateDependencyModuleNames.AddRange(
 			new string[] {
 				"Core",
+				"CorePreciseFP",
 		});
 	}
 }
diff --git a/Engine/Plugins/Experimental/PlainProps/Source/Private/PlainPropsDiff.cpp b/Engine/Plugins/Experimental/PlainProps/Source/Private/PlainPropsDiff.cpp
--- a/Engine/Plugins/Experimental/PlainProps/Source/Private/PlainPropsDiff.cpp
+++ b/Engine/Plugins/Experimental/PlainProps/Source/Private/PlainPropsDiff.cpp
@@ -45,7 +45,7 @@ static bool DiffMember(FRangeMemberBinding Member, const void* OwnerA, const voi
 		}
 		else if (InnerKind == EMemberKind::Leaf)
 		{
-			bDiff = DiffItemRange(A, B, ItemBinding, Ctx, SizeOf(GetBasicWidth(InnerType.AsLeaf())));
+			bDiff = DiffItemRange(A, B, ItemBinding, Ctx, SizeOf(GetItemWidth(InnerType.AsLeaf())));
 		}
 		else
 		{
diff --git a/Engine/Plugins/Experimental/PlainProps/Source/Private/PlainPropsInternalDiff.h b/Engine/Plugins/Experimental/PlainProps/Source/Private/PlainPropsInternalDiff.h
--- a/Engine/Plugins/Experimental/PlainProps/Source/Private/PlainPropsInternalDiff.h
+++ b/Engine/Plugins/Experimental/PlainProps/Source/Private/PlainPropsInternalDiff.h
@@ -4,6 +4,7 @@
 
 #include "PlainPropsBind.h"
 #include "PlainPropsTypes.h"
+#include "Math/PreciseFP.h"
 
 namespace PlainProps 
 {
@@ -15,27 +16,41 @@ inline const uint8* At(const void* Ptr, SIZE_T Offset)
 
 inline int64 DiffLeaf(const uint8* A, const uint8* B, FUnpackedLeafBindType Leaf)
 {
-	if (Leaf.Type == ELeafBindType::BitfieldBool)
+	bool bBitfield = Leaf.Type == ELeafBindType::BitfieldBool;
+	bool bFloat = Leaf.Type == ELeafBindType::Float;
+	if (bBitfield | bFloat)
 	{
-		uint32 Mask = 1u << Leaf.BitfieldIdx;
-		return (*A & Mask) - (*B & Mask);
-	}
+		if (bBitfield)
+		{
+			return ((*A ^ *B) >> Leaf.BitfieldIdx) & 1;
+		}
+		
+		switch (Leaf.Width)
+		{
+		case ELeafWidth::B32:	return !UE::PreciseFPEqual(reinterpret_cast<const float&>(*A), reinterpret_cast<const float&>(*B));
+		case ELeafWidth::B64:	return !UE::PreciseFPEqual(reinterpret_cast<const double&>(*A), reinterpret_cast<const double&>(*B));
+		default:				break;
+		}
 
-	switch (Leaf.Width)
+	}
+	else
 	{
-	case ELeafWidth::B8:	return FMemory::Memcmp(A, B, 1);
-	case ELeafWidth::B16:	return FMemory::Memcmp(A, B, 2);
-	case ELeafWidth::B32:	return FMemory::Memcmp(A, B, 4);
-	case ELeafWidth::B64:	return FMemory::Memcmp(A, B, 8);
+		switch (Leaf.Width)
+		{
+		case ELeafWidth::B8:	return FMemory::Memcmp(A, B, 1);
+		case ELeafWidth::B16:	return FMemory::Memcmp(A, B, 2);
+		case ELeafWidth::B32:	return FMemory::Memcmp(A, B, 4);
+		case ELeafWidth::B64:	return FMemory::Memcmp(A, B, 8);
+		}
 	}
 
 	unimplemented();
 	return 0;
 }
 
-inline ELeafWidth GetBasicWidth(FLeafBindType Leaf)
+inline ELeafWidth GetItemWidth(FLeafBindType Leaf)
 {
-	checkf(Leaf.Bind.Type != ELeafBindType::BitfieldBool, TEXT("Arrays of bitfields is not a thing"));
+	checkf(Leaf.Bind.Type != ELeafBindType::BitfieldBool, TEXT("Range of bitfields is illegal"));
 	return Leaf.Basic.Width;
 }
 
diff --git a/Engine/Plugins/Experimental/PlainProps/Source/Private/PlainPropsSave.cpp b/Engine/Plugins/Experimental/PlainProps/Source/Private/PlainPropsSave.cpp
--- a/Engine/Plugins/Experimental/PlainProps/Source/Private/PlainPropsSave.cpp
+++ b/Engine/Plugins/Experimental/PlainProps/Source/Private/PlainPropsSave.cpp
@@ -241,7 +241,7 @@ inline const FStructDeclaration& GetDeclaration(const FSaveContext& Ctx, FStruct
 	const IItemRangeBinding& ItemBinding = Binding.AsItemBinding();
 	switch (InnerType.GetKind())
 	{
-	case EMemberKind::Leaf:		return SaveLeafRange(Range, ItemBinding, Ctx, GetBasicWidth(InnerType.AsLeaf()));
+	case EMemberKind::Leaf:		return SaveLeafRange(Range, ItemBinding, Ctx, GetItemWidth(InnerType.AsLeaf()));
 	case EMemberKind::Range:	return SaveNonLeafRange<FInternalNestedRangeSaver>(Range, ItemBinding, Ctx, GetInnerRange(Member));
 	case EMemberKind::Struct:	return SaveStructRange(Range, ItemBinding, Ctx, static_cast<FStructSchemaId>(Member.InnermostSchema.Get()));
 	}
@@ -361,7 +361,7 @@ static bool DiffItem(const uint8* A, const uint8* B, const FSaveContext& Ctx, FR
 	const IItemRangeBinding& ItemBinding = Binding.AsItemBinding();
 	switch (InnerType.GetKind())
 	{
-	case EMemberKind::Leaf:		return DiffItemRange(A, B, ItemBinding, Ctx, SizeOf(GetBasicWidth(InnerType.AsLeaf())));
+	case EMemberKind::Leaf:		return DiffItemRange(A, B, ItemBinding, Ctx, SizeOf(GetItemWidth(InnerType.AsLeaf())));
 	case EMemberKind::Range:	return DiffItemRange(A, B, ItemBinding, Ctx, GetInnerRange(Member));
 	case EMemberKind::Struct:	return DiffItemRange(A, B, ItemBinding, Ctx, static_cast<FStructSchemaId>(Member.InnermostSchema.Get()));
 	}
diff --git a/Engine/Plugins/Experimental/PlainProps/Source/Public/PlainPropsBind.h b/Engine/Plugins/Experimental/PlainProps/Source/Public/PlainPropsBind.h
--- a/Engine/Plugins/Experimental/PlainProps/Source/Public/PlainPropsBind.h
+++ b/Engine/Plugins/Experimental/PlainProps/Source/Public/PlainPropsBind.h
@@ -161,7 +161,7 @@ struct FUnpackedLeafBindType
 	constexpr FUnpackedLeafBindType(FLeafBindType In)
 	: Type(In.Bind.Type)
 	{
-		if (Type == ELeafBindType::BitfieldBool)
+		if (Type != ELeafBindType::BitfieldBool)
 		{
 			Width = In.Basic.Width;
 		}
diff --git a/Engine/Plugins/Experimental/PlainPropsUObject/Source/Private/PlainPropsUObjectBindings.cpp b/Engine/Plugins/Experimental/PlainPropsUObject/Source/Private/PlainPropsUObjectBindings.cpp
--- a/Engine/Plugins/Experimental/PlainPropsUObject/Source/Private/PlainPropsUObjectBindings.cpp
+++ b/Engine/Plugins/Experimental/PlainPropsUObject/Source/Private/PlainPropsUObjectBindings.cpp
@@ -16,6 +16,8 @@
 #include "Misc/CommandLine.h"
 #include "Misc/CoreDelegates.h"
 #include "Misc/DefinePrivateMemberPtr.h"
+#include "Serialization/MemoryReader.h"
+#include "Serialization/MemoryWriter.h"
 #include "StructUtils/UserDefinedStruct.h"
 #include "UObject/AnsiStrProperty.h"
 #include "UObject/Class.h"
@@ -3312,6 +3314,26 @@ static void InitBatchedProperties()
 	GUE.Defaults.BindZeroes(GUE.Structs.LazyObjectPtr, sizeof(FLazyObjectPtr), alignof(FLazyObjectPtr));
 }
 
+#if UE_FNAME_OUTLINE_NUMBER
+static uint32 ToInt(FName In)
+{
+	return In.GetDisplayIndex().ToUnstableInt();
+}
+static FName FromInt(uint32 In)
+{
+	return FName::CreateFromDisplayId(FNameEntryId::FromUnstableInt(In), NAME_NO_NUMBER_INTERNAL);
+}
+#else
+static uint64 ToInt(FName In)
+{
+	return (uint64(In.GetNumber()) << 32) | In.GetDisplayIndex().ToUnstableInt();
+}
+static FName FromInt(uint64 In)
+{
+	return FName::CreateFromDisplayId(FNameEntryId::FromUnstableInt(static_cast<uint32>(In)), static_cast<uint32>(In >> 32));
+}
+#endif
+
 struct FMemoryPropertyBatch
 {
 	TArray<FText> Texts; // Tricky to serialize intrusively
@@ -3383,26 +3405,6 @@ struct FMemoryPropertyBatch
 		Out.ResetWeakPtr();
 		LoadSoleStruct(&Out.GetUniqueID(), In);
 	}
-
-#if UE_FNAME_OUTLINE_NUMBER
-	static uint32 ToInt(FName In)
-	{
-		return In.GetDisplayIndex().ToUnstableInt();
-	}
-	static FName FromInt(uint32 In)
-	{
-		return FName::CreateFromDisplayId(FNameEntryId::FromUnstableInt(In), NAME_NO_NUMBER_INTERNAL);
-	}
-#else
-	static uint64 ToInt(FName In)
-	{
-		return (uint64(In.GetNumber()) << 32) | In.GetDisplayIndex().ToUnstableInt();
-	}
-	static FName FromInt(uint64 In)
-	{
-		return FName::CreateFromDisplayId(FNameEntryId::FromUnstableInt(static_cast<uint32>(In)), static_cast<uint32>(In >> 32));
-	}
-#endif
 };
 
 inline bool DiffProperty(FName A, FName B) { return !A.IsEqual(B, ENameCase::CaseSensitive); }
@@ -3473,6 +3475,8 @@ struct TCustomPropertyBindings
 	TCustomPropertyBinding<FLazyObjectPtr, BatchType>	LazyObjectPtr;
 };
 
+////////////////////////////////////////////////////////////////////////////////////////////////
+
 struct FMemoryBatch
 {
 	TArray64<uint8>				Data;
@@ -3633,12 +3637,112 @@ private:
 	}
 };
 
+////////////////////////////////////////////////////////////////////////////////////////////////
+
+// Similar to PlainProps FMemoryBatch but for storing FArchive property serialization results
+struct FArchivedProperties
+{
+	TArray<uint8>			Data;
+	TArray<FText>			Texts;
+};
+
+static constexpr uint32 RoundtripPortFlags = PPF_UseDeprecatedProperties | PPF_ForceTaggedSerialization;
+
+// Match FMemoryPropertyBatch somewhat for a fair comparison, e.g. save FText on side and FName as integer
+class FPropertyWriter final : public FMemoryWriter
+{
+	TArray<FText>& Texts;
+public:
+	FPropertyWriter(FArchivedProperties& Out) 
+	: FMemoryWriter(Out.Data)
+	, Texts(Out.Texts)
+	{
+		SetPortFlags(RoundtripPortFlags);
+	}
+
+	FORCENOINLINE void WriteProperties(UObject* Object, UObject* Defaults)
+	{
+		UClass* Class = Object->GetClass();
+		Class->SerializeTaggedProperties(*this, reinterpret_cast<uint8*>(Object), Class, reinterpret_cast<uint8*>(Defaults));
+	}
+
+	virtual FArchive& operator<<(FText& Value) override
+	{
+		int32 Idx = INDEX_NONE;
+		if (!Value.IsEmpty())
+		{
+			Idx = Texts.Num();
+			Texts.Add(Value);
+		}
+		return WriteValue(Idx);
+	}
+	virtual FArchive& operator<<(FName& Value) override				{ return WriteValue(ToInt(Value));}
+	virtual FArchive& operator<<(UObject*& Value) override			{ return WriteValue(reinterpret_cast<uint64&>(Value)); }
+	virtual FArchive& operator<<(FObjectPtr& Value) override		{ return WriteValue(reinterpret_cast<uint64&>(Value)); }
+	virtual FArchive& operator<<(FWeakObjectPtr& Value) override	{ return WriteValue(reinterpret_cast<uint64&>(Value)); }
+	virtual FArchive& operator<<(FLazyObjectPtr& Value) override	{ return WriteValue(Value.GetUniqueID()); }
+	virtual FArchive& operator<<(FSoftObjectPtr& Value) override	{ Value.GetUniqueID().SerializePath(*this); return *this; }
+	virtual FArchive& operator<<(FSoftObjectPath& Value) override	{ Value.SerializePath(*this); return *this; }
+	virtual FString GetArchiveName() const override					{ return "FPropertyWriter"; }
+
+	template<typename T>
+	FArchive& WriteValue(T Value)
+	{
+		return *this << Value;
+	}
+};
+
+class FPropertyReader final : public FMemoryReader
+{
+	const TArray<FText>& Texts;
+public:
+	FPropertyReader(const FArchivedProperties& Out)
+	: FMemoryReader(Out.Data)
+	, Texts(Out.Texts)
+	{
+		SetPortFlags(RoundtripPortFlags);
+	}
+
+	FORCENOINLINE void ReadProperties(UObject* Object)
+	{
+		UClass* Class = Object->GetClass();
+		Class->SerializeTaggedProperties(*this, reinterpret_cast<uint8*>(Object), Class, nullptr);
+	}
+
+	virtual FArchive& operator<<(FText& Value) override
+	{ 
+		int32 Idx = ReadValue<int32>();
+		Value = Idx == INDEX_NONE ? FText::GetEmpty() : Texts[Idx];
+		return *this;
+	}
+
+	virtual FArchive& operator<<(FName& Value) override				{ Value = FromInt(ReadValue<decltype(ToInt(FName()))>()); return *this; }
+	virtual FArchive& operator<<(UObject*& Value) override			{ reinterpret_cast<uint64&>(Value) = ReadValue<uint64>(); return *this; }
+	virtual FArchive& operator<<(FObjectPtr& Value) override		{ reinterpret_cast<uint64&>(Value) = ReadValue<uint64>(); return *this; }
+	virtual FArchive& operator<<(FWeakObjectPtr& Value) override	{ reinterpret_cast<uint64&>(Value) = ReadValue<uint64>(); return *this; }
+	virtual FArchive& operator<<(FLazyObjectPtr& Value) override	{ Value.ResetWeakPtr(); Value.GetUniqueID() = ReadValue<FUniqueObjectGuid>(); return *this; }
+	virtual FArchive& operator<<(FSoftObjectPtr& Value) override	{ Value.ResetWeakPtr(); Value.GetUniqueID().SerializePath(*this); return *this; }
+	virtual FArchive& operator<<(FSoftObjectPath& Value) override	{ Value.SerializePath(*this); return *this; }
+	virtual FString GetArchiveName() const override					{ return "FPropertyReader"; }
+
+	template<typename T>
+	T ReadValue()
+	{
+		T Out;
+		*this << Out;
+		return Out;
+	}
+};
+
+////////////////////////////////////////////////////////////////////////////////////////////////
+
 struct FInstance
 {
 	UObject* Orig;
 	UObject* Arch;
 	UObject* PP;
 	UObject* TPS;
+	UObject* UPS;
 	FStructSchemaId Id;
 
 	void Init()
@@ -3721,52 +3825,101 @@ static int32 Roundtrip()
 		TPS_N.SetNumber(1 + &Instance - Instances.GetData());
 		Instance.TPS = MakeEmptyInstance(Instance.Orig, TPS_N);
 	}
+	FName UPS_N("UPS", 1);
+	for (FInstance& Instance : Instances)
+	{
+		UPS_N.SetNumber(1 + &Instance - Instances.GetData());
+		Instance.UPS = MakeEmptyInstance(Instance.Orig, UPS_N);
+	}
 
 	// Save
 	UE_LOGFMT(LogPlainPropsUObject, Display, "Saving UObjects to PlainProps...");
-	FMemoryBatch Doc;
-	TCustomPropertyBindings<FMemoryPropertyBatch> FlatBindings(Doc.Properties, GUE.Customs);
-	TCustomPropertyBindings<FMemoryPropertyBatch> DeltaBindings(Doc.Properties, GUE.DeltaCustoms);
+	FMemoryBatch Plain;
+	TCustomPropertyBindings<FMemoryPropertyBatch> FlatBindings(Plain.Properties, GUE.Customs);
+	TCustomPropertyBindings<FMemoryPropertyBatch> DeltaBindings(Plain.Properties, GUE.DeltaCustoms);
 	{
 		FBatchSaver Batch(FlatBindings.Overlay, DeltaBindings.Overlay, GUObjectArray.GetObjectArrayNum());
 		for (const FInstance& Instance : Instances)
 		{	
 			Batch.Save(Instance.Id, Instance.Orig, Instance.Arch ? Instance.PP : nullptr);
 		}
-		Doc.Data = Batch.Write(/* out */ Doc.RuntimeIds);
+		Plain.Data = Batch.Write(/* out */ Plain.RuntimeIds);
 	}
 
 	// Load
 	UE_LOGFMT(LogPlainPropsUObject, Display, "Loading UObjects from PlainProps...");
-	FMemoryBatchLoader Batch(DeltaBindings.Overlay, MakeMemoryView(Doc.Data), Doc.RuntimeIds);
+	FMemoryBatchLoader Batch(DeltaBindings.Overlay, MakeMemoryView(Plain.Data), Plain.RuntimeIds);
 	for (const FInstance& Instance : Instances)
 	{
 		Batch.Load(Instance.PP);
 	}
 
 	// Diff original vs PlainProps
-	UE_LOGFMT(LogPlainPropsUObject, Display, "Diffing UObjects...");
+	UE_LOGFMT(LogPlainPropsUObject, Display, "Diffing UObjects roundtripped via PlainProps...");
 	FDiffContext FlatCtx	= { GUE.Types, GUE.Schemas, FlatBindings.Overlay };
 	FDiffContext DeltaCtx	= { GUE.Types, GUE.Schemas, DeltaBindings.Overlay };
-	FDiffPath Path;
-	TUtf8StringBuilder<256> Diffs;
-	uint32 NumDiffs = 0;
+	TUtf8StringBuilder<256> PpDiffs;
+	uint32 NumPpDiffs = 0;
 	for (FInstance& Instance : Instances)
 	{
 		FDiffContext& Ctx = Instance.Arch ? DeltaCtx : FlatCtx;
 		if (DiffStructs(Instance.Id, Instance.Orig, Instance.PP, /* in-out */ Ctx))
 		{
-			PrintDiff(/* out */ Diffs, Ctx.Out, GUE.Names);
+			PrintDiff(/* out */ PpDiffs, Ctx.Out, GUE.Names);
 			//Batch.Reload(Instance.TPS, &Instance - Instances.GetData());
 			//UObject* Foo = MakeEmptyInstance(Instance.Orig, "Banana");
 			Ctx.Out.Reset();
-			Diffs.Append("\n");	
-			++NumDiffs;
+			PpDiffs.Append("\n");	
+			++NumPpDiffs;
+		}
+	}
+	UE_LOGFMT(LogPlainPropsUObject, Display, "Detected {Diffs} diffs in {Objs} UObjects saved in a {KB}KB value stream using PlainProps", NumPpDiffs, Instances.Num(), Plain.Data.NumBytes() / 1024);
+
+	UE_LOGFMT(LogPlainPropsUObject, Display, "Saving UObjects to TPS archive...");
+	FArchivedProperties Tps;
+	{
+		FPropertyWriter Archive(/* out */ Tps);
+		for (const FInstance& Instance : Instances)
+		{
+			Archive.WriteProperties(Instance.Orig, Instance.Arch ? Instance.TPS : nullptr);
+		}
+	}
+
+	UE_LOGFMT(LogPlainPropsUObject, Display, "Loading UObjects from TPS archive...");
+	{
+		FPropertyReader Archive(Tps);
+		for (const FInstance& Instance : Instances)
+		{
+			Archive.ReadProperties(Instance.TPS);
+		}
+	}
+
+	// Diff original vs TPS
+	UE_LOGFMT(LogPlainPropsUObject, Display, "Diffing UObjects roundtripped via TPS...");
+	TUtf8StringBuilder<256> TpsDiffs;
+	TArray<int32> TpsDiffIdxs;
+	for (FInstance& Instance : Instances)
+	{
+		static const FName SkipClasses[] = { "BodySetup" }; // Skips some structs due to native FCollisionResponse::operator==
+		if (Algo::Find(SkipClasses, Instance.Orig->GetClass()->GetFName()))
+		{
+			continue;
+		}
+
+		FDiffContext& Ctx = Instance.Arch ? DeltaCtx : FlatCtx;
+		if (DiffStructs(Instance.Id, Instance.Orig, Instance.TPS, /* in-out */ Ctx))
+		{
+			PrintDiff(/* out */ TpsDiffs, Ctx.Out, GUE.Names);
+			
+			FPropertyWriter(/* out */ Tps).WriteProperties(Instance.Orig, Instance.Arch ? Instance.UPS : nullptr);
+			Ctx.Out.Reset();
+			TpsDiffs.Append("\n");	
+			TpsDiffIdxs.Add(&Instance - &Instances[0]);
 		}
 	}
+	UE_LOGFMT(LogPlainPropsUObject, Display, "Detected {Diffs} diffs in {Objs} UObjects saved in a {KB}KB value stream using TPS", TpsDiffIdxs.Num(), Instances.Num(), Tps.Data.NumBytes() / 1024);
 
-	UE_LOGFMT(LogPlainPropsUObject, Display, "Done - {Diffs} diffs in {Objs} UObjects saved in a {KB}KB value stream", NumDiffs, Instances.Num(), Doc.Data.NumBytes() / 1024);
-	return Diffs.Len();
+	return NumPpDiffs;
 }
 
 static int32 TestBindings()
