Force the condition in CA_ASSUME (and therefore check macros) to be referenced, so that users don't need to mark their otherwise-unused variables with [[maybe_unused]].

diff --git a/Engine/Plugins/Animation/PoseSearch/Source/Runtime/Private/PoseSearchInteractionIsland.cpp b/Engine/Plugins/Animation/PoseSearch/Source/Runtime/Private/PoseSearchInteractionIsland.cpp
--- a/Engine/Plugins/Animation/PoseSearch/Source/Runtime/Private/PoseSearchInteractionIsland.cpp
+++ b/Engine/Plugins/Animation/PoseSearch/Source/Runtime/Private/PoseSearchInteractionIsland.cpp
@@ -581,7 +581,9 @@ void FInteractionIsland::InjectToActor(const UObject* AnimContext, bool bAddTick
 
 void FInteractionIsland::AddSearchContext(const FInteractionSearchContext& SearchContext)
 {
+#if DO_CHECK
 	check(SearchContext.CheckForConsistency());
+#endif
 	check(IsInGameThread());
 	SearchContexts.Add(SearchContext);
 }
diff --git a/Engine/Plugins/Animation/PoseSearch/Source/Runtime/Private/PoseSearchInteractionSubsystem.cpp b/Engine/Plugins/Animation/PoseSearch/Source/Runtime/Private/PoseSearchInteractionSubsystem.cpp
--- a/Engine/Plugins/Animation/PoseSearch/Source/Runtime/Private/PoseSearchInteractionSubsystem.cpp
+++ b/Engine/Plugins/Animation/PoseSearch/Source/Runtime/Private/PoseSearchInteractionSubsystem.cpp
@@ -1122,7 +1122,9 @@ void UPoseSearchInteractionSubsystem::Tick(float DeltaSeconds)
 	// clearing up AvailabilitiesMap for the next frame
 	ResetAvailabilities();
 
+#if DO_CHECK
 	check(ValidateAllIslands());
+#endif
 }
 
 TStatId UPoseSearchInteractionSubsystem::GetStatId() const
diff --git a/Engine/Plugins/Animation/PoseSearch/Source/Runtime/Public/PoseSearch/PoseSearchIndex.h b/Engine/Plugins/Animation/PoseSearch/Source/Runtime/Public/PoseSearch/PoseSearchIndex.h
--- a/Engine/Plugins/Animation/PoseSearch/Source/Runtime/Public/PoseSearch/PoseSearchIndex.h
+++ b/Engine/Plugins/Animation/PoseSearch/Source/Runtime/Public/PoseSearch/PoseSearchIndex.h
@@ -167,7 +167,9 @@ struct FSearchIndexAsset
 
 	int32 GetPoseIndexFromTime(float Time, int32 SchemaSampleRate) const
 	{
+#if DO_CHECK
 		check(IsInitialized());
+#endif
 
 		const int32 NumPoses = GetNumPoses();
 		int32 PoseOffset = FMath::RoundToInt(SchemaSampleRate * Time) - FirstSampleIdx;
diff --git a/Engine/Plugins/Experimental/PlainProps/Source/Private/PlainPropsRead.cpp b/Engine/Plugins/Experimental/PlainProps/Source/Private/PlainPropsRead.cpp
--- a/Engine/Plugins/Experimental/PlainProps/Source/Private/PlainPropsRead.cpp
+++ b/Engine/Plugins/Experimental/PlainProps/Source/Private/PlainPropsRead.cpp
@@ -191,7 +191,9 @@ FMemberReader::FMemberReader(const FStructSchema& Schema, FByteReader Values, FR
 , NumInnerSchemas(Schema.NumInnerSchemas)
 #endif
 {
+#if DO_CHECK
 	check(InnerSchemaIdx <= NumInnerSchemas);
+#endif
 	checkf(NumRangeTypes != 0xFFFFu, TEXT("GrabRangeTypes() doesn't check for wrap-around"));
 
 	if (IsSparse)
@@ -262,12 +264,16 @@ inline void FMemberReader::SkipSchema(FMemberType InnermostType)
 		InnerSchemaIdx += InnermostType.AsLeaf().Type == ELeafType::Enum;
 	}
 
+#if DO_CHECK
 	check(InnerSchemaIdx <= NumInnerSchemas);
+#endif
 }
 
 FSchemaId FMemberReader::GrabInnerSchema()
 {
+#if DO_CHECK
 	check(InnerSchemaIdx < NumInnerSchemas);
+#endif
 	uint32 Idx = InnerSchemaIdx++;
 	return GetInnerSchemas()[Idx];
 }
diff --git a/Engine/Plugins/Experimental/PlainProps/Source/Public/PlainPropsRead.h b/Engine/Plugins/Experimental/PlainProps/Source/Public/PlainPropsRead.h
--- a/Engine/Plugins/Experimental/PlainProps/Source/Public/PlainPropsRead.h
+++ b/Engine/Plugins/Experimental/PlainProps/Source/Public/PlainPropsRead.h
@@ -60,7 +60,9 @@ public:
 
 	[[nodiscard]] const uint8* GrabBytes(uint64 NumBytes)
 	{
+#if DO_CHECK
 		check(It + NumBytes <= End);
+#endif
 		const uint8* Out = It;
 		It += NumBytes;
 		return Out;
@@ -89,11 +91,15 @@ public:
 	
 	[[nodiscard]] uint64 GrabVarIntU()
 	{
+#if DO_CHECK
 		check(It < End);
+#endif
 		uint32 NumBytesRead;
 		uint64 Out = ReadVarUInt(It, NumBytesRead);
 		It += NumBytesRead;
+#if DO_CHECK
 		checkSlow(It <= End);
+#endif
 		return Out;
 	}
 
@@ -110,9 +116,24 @@ public:
 		}
 	}
 	
-	void CheckEmpty() const { check(It == End); }
-	void CheckNonEmpty() const { check(It != End); }
-	void CheckSize(int64 ExpectedSize) const { check(End - It == ExpectedSize); }
+	void CheckEmpty() const
+	{
+#if DO_CHECK
+		check(It == End);
+#endif
+	}
+	void CheckNonEmpty() const
+	{
+#if DO_CHECK
+		check(It != End);
+#endif
+	}
+	void CheckSize(int64 ExpectedSize) const
+	{
+#if DO_CHECK
+		check(End - It == ExpectedSize);
+#endif
+	}
 };
 
 // Helper that consumes 8 bits from the byte value stream
diff --git a/Engine/Plugins/Experimental/PlainProps/Source/Public/PlainPropsSaveMember.h b/Engine/Plugins/Experimental/PlainProps/Source/Public/PlainPropsSaveMember.h
--- a/Engine/Plugins/Experimental/PlainProps/Source/Public/PlainPropsSaveMember.h
+++ b/Engine/Plugins/Experimental/PlainProps/Source/Public/PlainPropsSaveMember.h
@@ -283,14 +283,18 @@ protected:
 	template<typename BuiltItemType>
 	inline void AddBuiltItem(BuiltItemType Item)
 	{
+#if DO_CHECK
 		check(It < End);
+#endif
 		reinterpret_cast<BuiltItemType&>(*It) = Item;
 		It += sizeof(Item);
 	}
 
 	[[nodiscard]] FTypedRange Finalize(FMemberSchema RangeSchema)
 	{
+#if DO_CHECK
 		check(It == End);
+#endif
 		return { RangeSchema, Range };
 	}
 
diff --git a/Engine/Plugins/Experimental/PlainPropsUObject/Source/Private/PlainPropsUObjectBindings.cpp b/Engine/Plugins/Experimental/PlainPropsUObject/Source/Private/PlainPropsUObjectBindings.cpp
--- a/Engine/Plugins/Experimental/PlainPropsUObject/Source/Private/PlainPropsUObjectBindings.cpp
+++ b/Engine/Plugins/Experimental/PlainPropsUObject/Source/Private/PlainPropsUObjectBindings.cpp
@@ -166,7 +166,9 @@ void FDefaultStructs::Bind(FStructSchemaId Id, const UScriptStruct* Struct)
 	UScriptStruct::ICppStructOps* Ops = Struct->GetCppStructOps();
 
 	ReserveFlags(Id.Idx);
+#if DO_CHECK
 	checkf(Flip(Bound[Id.Idx]), TEXT("'%s' already bound"), *GUE.Names.Print(Id));
+#endif
 
 	if (const UUserDefinedStruct* UserStruct = Cast<UUserDefinedStruct>(Struct))
 	{
@@ -206,14 +208,18 @@ void FDefaultStructs::Bind(FStructSchemaId Id, const UScriptStruct* Struct)
 void FDefaultStructs::BindZeroes(FStructSchemaId Id, SIZE_T Size, uint32 Alignment)
 {
 	ReserveFlags(Id.Idx);
+#if DO_CHECK
 	checkf(Flip(Bound[Id.Idx]), TEXT("'%s' already bound"), *GUE.Names.Print(Id));
+#endif
 	ReserveZeroes(Zeroes, Size, Alignment);
 }
 
 void FDefaultStructs::BindStatic(FStructSchemaId Id, const void* Struct)
 {
 	ReserveFlags(Id.Idx);
+#if DO_CHECK
 	checkf(Flip(Bound[Id.Idx]), TEXT("'%s' already bound"), *GUE.Names.Print(Id));
+#endif
 	check(!Instanced[Id.Idx]);
 	check(GetInstance(MakeStaticInstance(Struct)) == Struct);
 
@@ -223,7 +229,9 @@ void FDefaultStructs::BindStatic(FStructSchemaId Id, const void* Struct)
 
 void FDefaultStructs::Drop(FStructSchemaId Id)
 {
+#if DO_CHECK
 	checkf(!Flip(Bound[Id.Idx]), TEXT("'%s' isn't bound"), *GUE.Names.Print(Id));
+#endif
 	if (Instanced[Id.Idx])
 	{
 		Instanced[Id.Idx] = false;
@@ -233,7 +241,9 @@ void FDefaultStructs::Drop(FStructSchemaId Id)
 
 const void* FDefaultStructs::Get(FStructSchemaId Id)
 {
+#if DO_CHECK
 	checkf(Bound[Id.Idx], TEXT("'%s' lack default"), *GUE.Names.Print(Id));
+#endif
 	return Instanced[Id.Idx] ? GetInstance(Instances.FindChecked(Id)) : Zeroes.GetData(); 
 }
 
diff --git a/Engine/Source/Runtime/Analytics/TelemetryUtils/Private/TelemetryRouter.cpp b/Engine/Source/Runtime/Analytics/TelemetryUtils/Private/TelemetryRouter.cpp
--- a/Engine/Source/Runtime/Analytics/TelemetryUtils/Private/TelemetryRouter.cpp
+++ b/Engine/Source/Runtime/Analytics/TelemetryUtils/Private/TelemetryRouter.cpp
@@ -10,7 +10,9 @@ thread_local bool TelemetryRouterReentrancyGuard = false;
 
 void FTelemetryRouter::CheckNotReentrant()
 {
+#if DO_CHECK
     check(TelemetryRouterReentrancyGuard == 0);
+#endif
 }
 
 FTelemetryRouter::FTelemetryRouter()
diff --git a/Engine/Source/Runtime/Core/Public/Algo/KahnTopologicalSort.h b/Engine/Source/Runtime/Core/Public/Algo/KahnTopologicalSort.h
--- a/Engine/Source/Runtime/Core/Public/Algo/KahnTopologicalSort.h
+++ b/Engine/Source/Runtime/Core/Public/Algo/KahnTopologicalSort.h
@@ -132,7 +132,9 @@ namespace Algo
 				int32 RemainingNum = RemainingVertices.Num();
 #endif
 				RemainingVertices.Remove(Vertex);
+#if DO_CHECK
 				check(RemainingVertices.Num() == RemainingNum - 1); // Confirm Vertex was in RemainingVertices
+#endif
 				SortedRange.Add(Vertex);
 			}
 			Swap(NewIndependentVertices, IndependentVertices);
diff --git a/Engine/Source/Runtime/Core/Public/Misc/CoreMiscDefines.h b/Engine/Source/Runtime/Core/Public/Misc/CoreMiscDefines.h
--- a/Engine/Source/Runtime/Core/Public/Misc/CoreMiscDefines.h
+++ b/Engine/Source/Runtime/Core/Public/Misc/CoreMiscDefines.h
@@ -123,7 +123,7 @@
 	#define CA_CHECK_RETVAL
 	#define CA_NO_RETURN
 	#define CA_SUPPRESS( WarningNumber )
-	#define CA_ASSUME( Expr )
+	#define CA_ASSUME( Expr ) ((void)sizeof((bool)(Expr)))
 	#define CA_CONSTANT_IF(Condition) if (Condition)
 #endif
 
diff --git a/Engine/Source/Runtime/CoreUObject/Private/Serialization/AsyncLoading.cpp b/Engine/Source/Runtime/CoreUObject/Private/Serialization/AsyncLoading.cpp
--- a/Engine/Source/Runtime/CoreUObject/Private/Serialization/AsyncLoading.cpp
+++ b/Engine/Source/Runtime/CoreUObject/Private/Serialization/AsyncLoading.cpp
@@ -4284,7 +4284,9 @@ void FAsyncLoadingThread::InsertPackage(FAsyncPackage* Package, bool bReinsert,
 			}
 		}
 	}
+#if DO_CHECK || USING_CODE_ANALYSIS
 	check(!GEventDrivenLoaderEnabled || GetPackage(WeakPtr) == Package);
+#endif
 }
 
 void FAsyncLoadingThread::AddToLoadedPackages(FAsyncPackage* Package)
diff --git a/Engine/Source/Runtime/CoreUObject/Private/Serialization/UnversionedPropertySerialization.cpp b/Engine/Source/Runtime/CoreUObject/Private/Serialization/UnversionedPropertySerialization.cpp
--- a/Engine/Source/Runtime/CoreUObject/Private/Serialization/UnversionedPropertySerialization.cpp
+++ b/Engine/Source/Runtime/CoreUObject/Private/Serialization/UnversionedPropertySerialization.cpp
@@ -793,7 +793,9 @@ public:
 
 		FUnversionedPropertySerializer GetSerializer() const
 		{
+#if DO_CHECK || USING_CODE_ANALYSIS
 			check(SchemaIt != SchemaEnd);
+#endif
 			return *SchemaIt;
 		}
 
diff --git a/Engine/Source/Runtime/D3D12RHI/Private/D3D12CommandContext.cpp b/Engine/Source/Runtime/D3D12RHI/Private/D3D12CommandContext.cpp
--- a/Engine/Source/Runtime/D3D12RHI/Private/D3D12CommandContext.cpp
+++ b/Engine/Source/Runtime/D3D12RHI/Private/D3D12CommandContext.cpp
@@ -587,7 +587,9 @@ FD3D12CopyScope::FD3D12CopyScope(FD3D12Device* Device, ED3D12SyncPointType SyncP
 
 FD3D12CopyScope::~FD3D12CopyScope()
 {
+#if DO_CHECK
 	checkf(bSyncPointRetrieved, TEXT("The copy sync point must be retrieved before the end of the scope."));
+#endif
 
 	Context.SignalSyncPoint(SyncPoint);
 
diff --git a/Engine/Source/Runtime/D3D12RHI/Private/D3D12CommandList.cpp b/Engine/Source/Runtime/D3D12RHI/Private/D3D12CommandList.cpp
--- a/Engine/Source/Runtime/D3D12RHI/Private/D3D12CommandList.cpp
+++ b/Engine/Source/Runtime/D3D12RHI/Private/D3D12CommandList.cpp
@@ -52,7 +52,9 @@ void FD3D12CommandList::AddToResidencySet(TConstArrayView<FD3D12ResidencyHandle*
 	{
 		if (D3DX12Residency::IsInitialized(Handle))
 		{
+#if DO_CHECK
 			check(Device->GetGPUMask() == Handle->GPUObject->GetGPUMask());
+#endif
 			D3DX12Residency::Insert(*ResidencySet, *Handle);
 		}
 	}
diff --git a/Engine/Source/Runtime/D3D12RHI/Private/D3D12RayTracing.cpp b/Engine/Source/Runtime/D3D12RHI/Private/D3D12RayTracing.cpp
--- a/Engine/Source/Runtime/D3D12RHI/Private/D3D12RayTracing.cpp
+++ b/Engine/Source/Runtime/D3D12RHI/Private/D3D12RayTracing.cpp
@@ -1573,8 +1573,10 @@ public:
 
 		checkf(Data.Num(), TEXT("Shader table is expected to be initialized before copying to GPU."));
 
+#if DO_CHECK
 		checkf(bWasDefaultMissShaderSet, TEXT("At least the first miss shader must have been set before copying to GPU."));
-		
+#endif
+
 #if ENABLE_RESIDENCY_MANAGEMENT
 		// Merge all data from worker threads into the main set
 		{
diff --git a/Engine/Source/Runtime/Engine/Private/Components/ActorComponent.cpp b/Engine/Source/Runtime/Engine/Private/Components/ActorComponent.cpp
--- a/Engine/Source/Runtime/Engine/Private/Components/ActorComponent.cpp
+++ b/Engine/Source/Runtime/Engine/Private/Components/ActorComponent.cpp
@@ -1646,7 +1646,9 @@ void UActorComponent::PreRegisterComponentWithWorld(UWorld* InWorld)
 	checkf(WorldPrivate == nullptr, TEXT("%s"), *GetFullName());
 
 	AActor* MyOwner = GetOwner();
+#if DO_CHECK || USING_CODE_ANALYSIS
 	checkSlow(MyOwner == nullptr || MyOwner->OwnsComponent(this));
+#endif
 
 	if (MyOwner && MyOwner->GetClass()->HasAnyClassFlags(CLASS_NewerVersionExists))
 	{
@@ -1710,7 +1712,9 @@ void UActorComponent::RegisterComponentWithWorld(UWorld* InWorld, FRegisterCompo
 	checkf(WorldPrivate == nullptr || (IsPreRegistered()), TEXT("%s"), *GetFullName());
 
 	AActor* MyOwner = GetOwner();
+#if DO_CHECK || USING_CODE_ANALYSIS
 	checkSlow(MyOwner == nullptr || MyOwner->OwnsComponent(this));
+#endif
 
 	if (MyOwner && MyOwner->GetClass()->HasAnyClassFlags(CLASS_NewerVersionExists))
 	{
diff --git a/Engine/Source/Runtime/Engine/Private/SceneView.cpp b/Engine/Source/Runtime/Engine/Private/SceneView.cpp
--- a/Engine/Source/Runtime/Engine/Private/SceneView.cpp
+++ b/Engine/Source/Runtime/Engine/Private/SceneView.cpp
@@ -1067,7 +1067,9 @@ PRAGMA_ENABLE_DEPRECATION_WARNINGS
 		}
 	}
 
+#if DO_CHECK || USING_CODE_ANALYSIS
 	check(VerifyMembersChecks());
+#endif
 }
 
 #if !(UE_BUILD_SHIPPING || UE_BUILD_TEST)
@@ -2344,7 +2346,9 @@ void FSceneView::EndFinalPostprocessSettings(const FSceneViewInitOptions& ViewIn
 	}
 #endif // WITH_EDITOR
 
+#if DO_CHECK || USING_CODE_ANALYSIS
 	check(VerifyMembersChecks());
+#endif
 }
 
 void FSceneView::ConfigureBufferVisualizationSettings()
diff --git a/Engine/Source/Runtime/Engine/Public/InstanceDataSceneProxy.h b/Engine/Source/Runtime/Engine/Public/InstanceDataSceneProxy.h
--- a/Engine/Source/Runtime/Engine/Public/InstanceDataSceneProxy.h
+++ b/Engine/Source/Runtime/Engine/Public/InstanceDataSceneProxy.h
@@ -247,7 +247,9 @@ public:
 	{
 		check(AccessTag.Kind == FAccessTag::EKind::Writer && AccessTag.WriterTag != 0u);
 		uint32 PrevTagValue = 0u;
+#if DO_CHECK
 		check(CurrentWriterTag.Value.compare_exchange_strong(PrevTagValue, AccessTag.WriterTag));
+#endif
 		return FWriteView
 		{
 			AccessTag,
@@ -277,7 +279,9 @@ public:
 		check(bInstanceDataIsGPUOnly ? InstanceToPrimitiveRelative.IsEmpty() : (NumInstancesGPUOnly == 0));
 		check(AccessTag.Kind == FAccessTag::EKind::Writer && AccessTag.WriterTag != 0u);
 		uint32 PrevTagValue = AccessTag.WriterTag;
+#if DO_CHECK
 		check(CurrentWriterTag.Value.compare_exchange_strong(PrevTagValue, 0U));
+#endif
 	}
 
 
diff --git a/Engine/Source/Runtime/RHI/Private/RHICommandList.cpp b/Engine/Source/Runtime/RHI/Private/RHICommandList.cpp
--- a/Engine/Source/Runtime/RHI/Private/RHICommandList.cpp
+++ b/Engine/Source/Runtime/RHI/Private/RHICommandList.cpp
@@ -237,8 +237,10 @@ void FRHICommandListBase::AttachBreadcrumbSubTree(FRHIBreadcrumbAllocator& Alloc
 
 void FRHICommandListBase::ActivatePipelines(ERHIPipeline Pipelines)
 {
+#if DO_CHECK
 	checkf(IsTopOfPipe() || Bypass(), TEXT("Cannot be called from the bottom of pipe."));
 	checkf(Pipelines == ERHIPipeline::None || EnumHasAllFlags(AllowedPipelines, Pipelines), TEXT("At least one of the specified pipelinea are not allowed on this RHI command list."));
+#endif
 
 	if (ActivePipelines == Pipelines)
 	{
diff --git a/Engine/Source/Runtime/RHI/Private/RHIResources.cpp b/Engine/Source/Runtime/RHI/Private/RHIResources.cpp
--- a/Engine/Source/Runtime/RHI/Private/RHIResources.cpp
+++ b/Engine/Source/Runtime/RHI/Private/RHIResources.cpp
@@ -33,9 +33,9 @@ FRHIResource::FRHIResource(ERHIResourceType InResourceType)
 
 FRHIResource::~FRHIResource()
 {
+#if DO_CHECK
 	check(IsEngineExitRequested() || CurrentlyDeleting == this);
 	check(AtomicFlags.GetNumRefs(std::memory_order_relaxed) == 0); // this should not have any outstanding refs
-#if DO_CHECK
 	CurrentlyDeleting = nullptr;
 #endif
 
@@ -70,7 +70,9 @@ void FRHIResource::DeleteResources(TArray<FRHIResource*> const& Resources)
 #endif
 			delete Resource;
 
+#if DO_CHECK
 			check(CurrentlyDeleting == nullptr);
+#endif
 		}
 	}
 }
diff --git a/Engine/Source/Runtime/RHI/Public/RHITransition.h b/Engine/Source/Runtime/RHI/Public/RHITransition.h
--- a/Engine/Source/Runtime/RHI/Public/RHITransition.h
+++ b/Engine/Source/Runtime/RHI/Public/RHITransition.h
@@ -885,7 +885,9 @@ private:
 
 	inline void MarkBegin(ERHIPipeline Pipeline) const
 	{
+#if DO_CHECK || USING_CODE_ANALYSIS
 		checkf(EnumHasAllFlags(AllowedSrc, Pipeline), TEXT("Transition is being used on a source pipeline that it wasn't created for."));
+#endif
 
 		int8 Mask = int8(Pipeline);
 		int8 PreviousValue = FPlatformAtomics::InterlockedAnd(&State, ~Mask);
@@ -899,7 +901,9 @@ private:
 
 	inline void MarkEnd(ERHIPipeline Pipeline) const
 	{
+#if DO_CHECK || USING_CODE_ANALYSIS
 		checkf(EnumHasAllFlags(AllowedDst, Pipeline), TEXT("Transition is being used on a destination pipeline that it wasn't created for."));
+#endif
 
 		int8 Mask = int8(int32(Pipeline) << int32(ERHIPipeline::Num));
 		int8 PreviousValue = FPlatformAtomics::InterlockedAnd(&State, ~Mask);
diff --git a/Engine/Source/Runtime/Renderer/Private/PostProcess/PostProcessing.cpp b/Engine/Source/Runtime/Renderer/Private/PostProcess/PostProcessing.cpp
--- a/Engine/Source/Runtime/Renderer/Private/PostProcess/PostProcessing.cpp
+++ b/Engine/Source/Runtime/Renderer/Private/PostProcess/PostProcessing.cpp
@@ -362,7 +362,9 @@ void AddPostProcessingPasses(
 	using namespace UE::Renderer::PostProcess;
 
 	check(IsInRenderingThread());
+#if DO_CHECK || USING_CODE_ANALYSIS
 	check(View.VerifyMembersChecks());
+#endif
 	Inputs.Validate();
 
 	FScene* Scene = View.Family->Scene->GetRenderScene();
@@ -2059,7 +2061,9 @@ void AddDebugViewPostProcessingPasses(FRDGBuilder& GraphBuilder, const FViewInfo
 	QUICK_SCOPE_CYCLE_COUNTER(STAT_PostProcessing_Process);
 
 	check(IsInRenderingThread());
+#if DO_CHECK || USING_CODE_ANALYSIS
 	check(View.VerifyMembersChecks());
+#endif
 	Inputs.Validate();
 
 	const FIntRect PrimaryViewRect = View.ViewRect;
@@ -2299,7 +2303,9 @@ void AddVisualizeCalibrationMaterialPostProcessingPasses(FRDGBuilder& GraphBuild
 	QUICK_SCOPE_CYCLE_COUNTER(STAT_PostProcessing_Process);
 
 	check(IsInRenderingThread());
+#if DO_CHECK || USING_CODE_ANALYSIS
 	check(View.VerifyMembersChecks());
+#endif
 	check(InMaterialInterface);
 	Inputs.Validate();
 
diff --git a/Engine/Source/Runtime/Renderer/Private/SceneCulling/SceneCulling.cpp b/Engine/Source/Runtime/Renderer/Private/SceneCulling/SceneCulling.cpp
--- a/Engine/Source/Runtime/Renderer/Private/SceneCulling/SceneCulling.cpp
+++ b/Engine/Source/Runtime/Renderer/Private/SceneCulling/SceneCulling.cpp
@@ -2862,9 +2862,13 @@ void FSceneCulling::FUpdater::PreSceneUpdate(FRDGBuilder& GraphBuilder, const FS
 			SCOPE_CYCLE_COUNTER(STAT_SceneCulling_Update_Pre);
 			BUILDER_LOG_SCOPE("OnPreSceneUpdate: %d", ScenePreUpdateData.RemovedPrimitiveIds.Num());
 			CSV_SCOPED_TIMING_STAT(SceneCulling, PreSceneUpdate);
+#if DO_CHECK
 			check(DebugTaskCounter++ == 0);
+#endif
 			SceneCulling.ActiveUpdaterImplementation->ProcessPreSceneUpdate(ChangeSet);
+#if DO_CHECK
 			check(--DebugTaskCounter == 0);
+#endif
 		}
 #if SC_ALLOW_ASYNC_TASKS
 		, SceneCulling.bUseAsyncUpdate);
@@ -2887,9 +2891,13 @@ void FSceneCulling::FUpdater::PostSceneUpdate(FRDGBuilder& GraphBuilder, const F
 	SceneCulling.PostUpdateTaskHandle = GraphBuilder.AddSetupTask([this, ChangeSet]()
 #endif
 	{
+#if DO_CHECK
 		check(DebugTaskCounter++ == 0);
+#endif
 		SceneCulling.ActiveUpdaterImplementation->ProcessPostSceneUpdate(ChangeSet);
+#if DO_CHECK
 		check(--DebugTaskCounter == 0);
+#endif
 	}
 #if SC_ALLOW_ASYNC_TASKS
 	, nullptr, TArray<UE::Tasks::FTask>{ PreUpdateTaskHandle }, UE::Tasks::ETaskPriority::Normal, SceneCulling.bUseAsyncUpdate);
diff --git a/Engine/Source/Runtime/Renderer/Private/SceneRenderBuilder.cpp b/Engine/Source/Runtime/Renderer/Private/SceneRenderBuilder.cpp
--- a/Engine/Source/Runtime/Renderer/Private/SceneRenderBuilder.cpp
+++ b/Engine/Source/Runtime/Renderer/Private/SceneRenderBuilder.cpp
@@ -1000,7 +1000,9 @@ FSceneRenderBuilder::~FSceneRenderBuilder()
 {
 	if (PersistentState)
 	{
+#if DO_CHECK
 		checkf(PersistentState->RenderersToAdd.IsEmpty(), TEXT("FSceneRenderBuilder::Execute called but %d renderers were not added."), PersistentState->RenderersToAdd.Num());
+#endif
 		delete PersistentState;
 		PersistentState = nullptr;
 	}
diff --git a/Engine/Source/Runtime/Renderer/Private/SceneRendering.cpp b/Engine/Source/Runtime/Renderer/Private/SceneRendering.cpp
--- a/Engine/Source/Runtime/Renderer/Private/SceneRendering.cpp
+++ b/Engine/Source/Runtime/Renderer/Private/SceneRendering.cpp
@@ -2776,7 +2776,9 @@ FSceneRenderer::FSceneRenderer(const FSceneViewFamily* InViewFamily, FHitProxyCo
 			{
 				View.AntiAliasingMethod = AAM_TemporalAA;
 				View.PrimaryScreenPercentageMethod = EPrimaryScreenPercentageMethod::TemporalUpscale;
+#if DO_CHECK || USING_CODE_ANALYSIS
 				check(View.VerifyMembersChecks());
+#endif
 			}
 		}
 
@@ -3212,7 +3214,9 @@ void FSceneRenderer::PrepareViewRectsForRendering()
 		checkf(PrimaryResolutionFraction <= DynamicResolutionUpperBounds[GDynamicPrimaryResolutionFraction],
 			TEXT("ISceneViewFamilyScreenPercentage::GetPrimaryResolutionFractionUpperBound() should not lie to the renderer."));
 
+#if DO_CHECK || USING_CODE_ANALYSIS
 		check(ISceneViewFamilyScreenPercentage::IsValidResolutionFraction(PrimaryResolutionFraction));
+#endif
 	}
 
 	float LensDistortionResolutionFraction = 1.0f;
@@ -3279,7 +3283,9 @@ void FSceneRenderer::PrepareViewRectsForRendering()
 		}
 
 		check(View.ViewRect.Area() != 0);
+#if DO_CHECK || USING_CODE_ANALYSIS
 		check(View.VerifyMembersChecks());
+#endif
 
 		if (!ViewFamily.EngineShowFlags.HitProxies && View.AntiAliasingMethod == AAM_TSR)
 		{
diff --git a/Engine/Source/Runtime/Renderer/Private/SceneUpdateCommandQueue.h b/Engine/Source/Runtime/Renderer/Private/SceneUpdateCommandQueue.h
--- a/Engine/Source/Runtime/Renderer/Private/SceneUpdateCommandQueue.h
+++ b/Engine/Source/Runtime/Renderer/Private/SceneUpdateCommandQueue.h
@@ -116,13 +116,17 @@ public:
 		: PayloadArrays(MoveTemp(Other.PayloadArrays))
 		, Commands(MoveTemp(Other.Commands))
 	{
+#if DO_CHECK
 		check(Other.RaceGuard == 0);
+#endif
 		Other.CommandSlots.Empty();
 	}
 
 	~TSceneUpdateCommandQueue()
 	{
+#if DO_CHECK
 		check(RaceGuard == 0);
+#endif
 	}
 
 	bool IsEmpty() const { return Commands.IsEmpty(); }
@@ -150,7 +154,9 @@ public:
 	 */
 	void EnqueueDelete(FSceneInfo* SceneInfo)
 	{
+#if DO_CHECK
 		check(RaceGuard == 0);
+#endif
 
 		int32 CommandSlot = GetOrAddCommandSlot(SceneInfo);
 		FUpdateCommand& Command = Commands[CommandSlot];
@@ -164,7 +170,9 @@ public:
 	 */
 	void EnqueueAdd(FSceneInfo* SceneInfo)
 	{
+#if DO_CHECK
 		check(RaceGuard == 0);
+#endif
 
 		// Add should always be the first command.
 		check(CommandSlots.Find(SceneInfo) == nullptr);
@@ -183,7 +191,9 @@ public:
 	template <typename PayloadType>
 	void Enqueue(FSceneInfo* SceneInfo, PayloadType&& Payload)
 	{
+#if DO_CHECK
 		check(RaceGuard == 0);
+#endif
 
 		int32 CommandSlot = GetOrAddCommandSlot(SceneInfo);
 		FUpdateCommand& Command = Commands[CommandSlot];
@@ -242,7 +252,9 @@ public:
 	 */
 	void Reset()
 	{
+#if DO_CHECK
 		check(RaceGuard == 0);
+#endif
 
 		CommandSlots.Empty();
 		Commands.Reset();
diff --git a/Engine/Source/Runtime/Renderer/Private/SceneVisibility.cpp b/Engine/Source/Runtime/Renderer/Private/SceneVisibility.cpp
--- a/Engine/Source/Runtime/Renderer/Private/SceneVisibility.cpp
+++ b/Engine/Source/Runtime/Renderer/Private/SceneVisibility.cpp
@@ -4869,7 +4869,9 @@ void FSceneRenderer::PrepareViewStateForVisibility(const FSceneTexturesConfig& S
 		FViewInfo& View = *AllViews[ViewIndex];
 		FSceneViewState* ViewState = View.ViewState;
 
+#if DO_CHECK || USING_CODE_ANALYSIS
 		check(View.VerifyMembersChecks());
+#endif
 
 		// Setup global dither fade in and fade out uniform buffers.
 		{
