PlainProps bind UStruct inheritance and reimplement map/set property delta bindinings

Add UStruct inheritance bindings, PlainProps inheritance binding and inheritance diffing
Reimplemented TMap/TSet property delta bindings
Fixed bugs related to dynamic structs

#rnx

diff --git a/Engine/Plugins/Experimental/PlainProps/Source/Private/PlainPropsBind.cpp b/Engine/Plugins/Experimental/PlainProps/Source/Private/PlainPropsBind.cpp
--- a/Engine/Plugins/Experimental/PlainProps/Source/Private/PlainPropsBind.cpp
+++ b/Engine/Plugins/Experimental/PlainProps/Source/Private/PlainPropsBind.cpp
@@ -115,6 +115,15 @@ void FMemberVisitor::SkipMember()
 	++MemberIdx;
 }
 
+FStructSchemaId FMemberVisitor::GrabSuper()
+{
+	check(MemberIdx == 0);
+	checkSlow(Schema.Members[0].AsByte() == SuperStructType.AsByte());
+	MemberIdx = 1;
+	InnerSchemaIdx = 1;
+	return static_cast<FStructSchemaId>(Schema.GetInnerSchemas()[0]);
+}
+
 FSchemaId FMemberVisitor::GrabInnerSchema()
 {
 	check(InnerSchemaIdx < Schema.NumInnerSchemas);
diff --git a/Engine/Plugins/Experimental/PlainProps/Source/Private/PlainPropsDiff.cpp b/Engine/Plugins/Experimental/PlainProps/Source/Private/PlainPropsDiff.cpp
--- a/Engine/Plugins/Experimental/PlainProps/Source/Private/PlainPropsDiff.cpp
+++ b/Engine/Plugins/Experimental/PlainProps/Source/Private/PlainPropsDiff.cpp
@@ -59,15 +59,15 @@ static bool DiffMember(FRangeMemberBinding Member, const void* OwnerA, const voi
 	return bDiff;
 }
 
-static bool DiffMembers(const FSchemaBinding& Schema, const void* A, const void* B, FDiffContext& Ctx);
+static bool DiffMembers(FStructSchemaId BindId, const void* A, const void* B, FDiffContext& Ctx);
 
-inline bool DiffStruct(FStructSchemaId Id, const void* A, const void* B, FDiffContext& Ctx)
+inline bool DiffStruct(FStructSchemaId BindId, const void* A, const void* B, FDiffContext& Ctx)
 {
-	if (const ICustomBinding* Custom = Ctx.Customs.FindStruct(Id))
+	if (const ICustomBinding* Custom = Ctx.Customs.FindStruct(BindId))
 	{
 		return Custom->DiffCustom(A, B);
 	}
-	return DiffMembers(Ctx.Schemas.GetStruct(Id), A, B, Ctx);
+	return DiffMembers(BindId, A, B, Ctx);
 }
 
 static bool DiffMember(FStructMemberBinding Member, const void* OwnerA, const void* OwnerB, FDiffContext& Ctx, FMemberId Name)
@@ -82,24 +82,42 @@ static bool DiffMember(FStructMemberBinding Member, const void* OwnerA, const vo
 	return false;
 }
 
-static bool DiffMembers(const FSchemaBinding& Schema, const void* A, const void* B, FDiffContext& Ctx)
+static bool DiffMember(/* in-out */ FMemberVisitor& It, const void* A, const void* B, FDiffContext& Ctx, FMemberId Name)
+{
+	switch (It.PeekKind())
+	{
+	case EMemberKind::Leaf:		return DiffMember(It.GrabLeaf(), A, B, Ctx, Name);
+	case EMemberKind::Range:	return DiffMember(It.GrabRange(), A, B, Ctx, Name);
+	case EMemberKind::Struct:	return DiffMember(It.GrabStruct(), A, B, Ctx, Name);
+	}
+	return true;
+}
+
+
+static bool DiffMembers(FStructSchemaId BindId, const void* A, const void* B, FDiffContext& Ctx)
 {
+	const FSchemaBinding& Schema = Ctx.Schemas.GetStruct(BindId);
 	const FStructDeclaration& Declaration = Ctx.Types.Get(Schema.DeclId);
-	TConstArrayView<FMemberId> MemberOrder = Declaration.GetMemberOrder();
+	
+	FMemberVisitor It(Schema);
+	if (Schema.HasSuper())
+	{
+		if (DiffMembers(It.GrabSuper(), A, B, Ctx))
+		{
+			return true;
+		}
+	}
 
-	bool bDiff = false;
-	for (FMemberVisitor It(Schema); !bDiff && It.HasMore(); )
+	for (FMemberId Name : Declaration.GetMemberOrder())
 	{
-		FMemberId Name = MemberOrder[It.GetIndex()];
-		switch (It.PeekKind())
+		if (DiffMember(/* in-out */It, A, B, Ctx, Name))
 		{
-			case EMemberKind::Leaf:		bDiff = DiffMember(It.GrabLeaf(), A, B, Ctx, Name);	break;
-			case EMemberKind::Range:	bDiff = DiffMember(It.GrabRange(), A, B, Ctx, Name);	break;
-			case EMemberKind::Struct:	bDiff = DiffMember(It.GrabStruct(), A, B, Ctx, Name);	break;
+			return true;
 		}
 	}
 
-	return bDiff;
+	check(!It.HasMore());
+	return false;
 }
 
 static bool DiffItem(const void* A, const void* B, FDiffContext& Ctx, FRangeMemberBinding Range)
@@ -121,7 +139,7 @@ bool DiffStructs(FStructSchemaId BindId, const void* A, const void* B, FDiffCont
 		// Currently no way for DiffCustom to emit deeper diff into
 		return Custom->DiffCustom(A, B);
 	}
-	return DiffMembers(Ctx.Schemas.GetStruct(BindId), A, B, Ctx);
+	return DiffMembers(BindId, A, B, Ctx);
 }
 
 } // namespace PlainProps
\ No newline at end of file
diff --git a/Engine/Plugins/Experimental/PlainProps/Source/Private/PlainPropsInternalBind.h b/Engine/Plugins/Experimental/PlainProps/Source/Private/PlainPropsInternalBind.h
--- a/Engine/Plugins/Experimental/PlainProps/Source/Private/PlainPropsInternalBind.h
+++ b/Engine/Plugins/Experimental/PlainProps/Source/Private/PlainPropsInternalBind.h
@@ -23,6 +23,7 @@ public:
 	FLeafMemberBinding			GrabLeaf();				// @pre PeekKind() == EMemberKind::Leaf
 	FRangeMemberBinding			GrabRange();			// @pre PeekKind() == EMemberKind::Range
 	FStructMemberBinding		GrabStruct();			// @pre PeekKind() == EMemberKind::Struct
+	FStructSchemaId				GrabSuper();			// @pre First grab and has declared super
 	void						SkipMember();
 
 protected: // for unit tests
diff --git a/Engine/Plugins/Experimental/PlainProps/Source/Private/PlainPropsInternalFormat.h b/Engine/Plugins/Experimental/PlainProps/Source/Private/PlainPropsInternalFormat.h
--- a/Engine/Plugins/Experimental/PlainProps/Source/Private/PlainPropsInternalFormat.h
+++ b/Engine/Plugins/Experimental/PlainProps/Source/Private/PlainPropsInternalFormat.h
@@ -50,6 +50,12 @@ static_assert(offsetof(FSchemaBatch, NumStructSchemas) == 16,					"Add binary fo
 
 //////////////////////////////////////////////////////////////////////////
 
+// Direct super schema id is always written to enable inheritance analysis
+// 
+// No		No inheritance
+// Unused	Super member missing, due to all instances being empty
+// Used		Super member exists, but IsDynamic or with a static non-direct ancestor id
+// Reused	Super member exists, with the static id of the direct super struct
 enum class ESuper : uint8 { No, Unused, Used, Reused };
 
 inline bool UsesSuper(ESuper Inheritance)				{ return Inheritance == ESuper::Used   || Inheritance == ESuper::Reused; }
diff --git a/Engine/Plugins/Experimental/PlainProps/Source/Private/PlainPropsLoad.cpp b/Engine/Plugins/Experimental/PlainProps/Source/Private/PlainPropsLoad.cpp
--- a/Engine/Plugins/Experimental/PlainProps/Source/Private/PlainPropsLoad.cpp
+++ b/Engine/Plugins/Experimental/PlainProps/Source/Private/PlainPropsLoad.cpp
@@ -172,9 +172,10 @@ struct FLoadIdMapping
 class FLoadIdBinder
 {
 public:
-	FLoadIdBinder(TConstArrayView<FStructSchemaId> RuntimeDeclIds)
+	FLoadIdBinder(TConstArrayView<FStructSchemaId> RuntimeDeclIds, FDebugIds Dbg)
 	: DeclIds(RuntimeDeclIds)
 	, NextLoadIdx(static_cast<uint32>(RuntimeDeclIds.Num()))
+	, Debug(Dbg)
 	{}
 
 	FStructSchemaId					BindLoadId(FStructSchemaId ReadId, FStructSchemaId BindId)
@@ -208,10 +209,13 @@ private:
 		Mapping.BindId = BindId;
 		return Mapping.LoadId;
 	}
+public:
+	FDebugIds							Debug;
 };
 
 struct FLoadIdBinderDummy
 {
+	FDebugIds							Debug;
 	static constexpr FStructSchemaId	BindLoadId(FStructSchemaId ReadId, FStructSchemaId BindId) { return ReadId; }
 	static constexpr int32				NumMappings() { return 0; }
 	static inline FLoadIdMapping		GetMapping(int32 Idx) { unimplemented(); return {}; }
@@ -259,6 +263,8 @@ static void CopyMemberBinding(FLeafMemberBinding Binding, /* in-out */ const FSc
 template<class LoadIdBinder>
 void CopyMemberBinding(FStructMemberBinding Binding, /* in-out */ const FSchemaId*& InnerSchemaIt, LoadIdBinder& LoadIds, FMemberBinder& Out)
 {
+	check(!Binding.Type.IsSuper);
+	check(!Binding.Type.IsDynamic); // TODO: Build out dynamic struct support
 	FStructSchemaId LoadId = LoadIds.BindLoadId(static_cast<FStructSchemaId>(*InnerSchemaIt), Binding.Id);
 	Out.AddMember(FMemberBindType(Binding.Type), static_cast<uint32>(Binding.Offset));
 	Out.AddInnerSchema(LoadId);
@@ -297,29 +303,53 @@ void CopyMemberBinding(/* in-out */ FMemberVisitor& BindIt, /* in-out */ const F
 template<class LoadIdBinder>
 static void CreateSubsetBindingWithoutEnumIds(const FStructSchema& From, const FSchemaBinding& To, TConstArrayView<FMemberId> ToNames, uint16 NumEnums, LoadIdBinder& LoadIds, SubsetByteArray& Out)
 {
-	check(To.NumMembers == ToNames.Num());
+	check(To.NumMembers == ToNames.Num() + To.HasSuper());
 	check(To.NumMembers >= From.NumMembers);
 	check(IsAligned(Out.Num(), alignof(FSchemaBinding)));
 
 	int32 OutPos = Out.Num();
 	
 	// Allocate and init header 
-	bool bSkipSuperSchema = From.Inheritance == ESuper::Unused;
+	bool bSkipSuperSchema = SkipDeclaredSuperSchema(From.Inheritance);
 	FSchemaBinding Header = { To.DeclId, From.NumMembers, From.NumInnerSchemas - NumEnums - bSkipSuperSchema, From.NumRangeTypes };
 	Out.AddUninitialized(Header.CalculateSize());
 	FSchemaBinding* Schema = new (&Out[OutPos]) FSchemaBinding {Header};
 	
-	// Copy subset of member bindings
+	// Copy subset of member bindings...
 	FMemberVisitor ToIt(To);
 	FMemberBinder Footer(*Schema);
 	const FSchemaId* InnerSchemaIt = From.GetInnerSchemas() + bSkipSuperSchema;
+
+	// ...first the unnamed super member...
+	if (From.Inheritance != ESuper::No)
+	{
+		FStructSchemaId BindId = ToIt.GrabSuper();
+		if (UsesSuper(From.Inheritance))
+		{
+			FStructType FromType = From.Footer[0].AsStruct(); // To.Members[0].AsStruct().IsDynamic isn't set
+			check(FromType.IsSuper);
+			Footer.AddMember(FMemberBindType(FromType), 0);
+		
+			if (!FromType.IsDynamic)
+			{
+				FStructSchemaId ReadId = static_cast<FStructSchemaId>(*InnerSchemaIt);
+				FStructSchemaId LoadId = LoadIds.BindLoadId(ReadId, BindId);
+				Footer.AddInnerSchema(LoadId);
+				++InnerSchemaIt;
+			}
+		}
+	}
+
+	// ...then a subset of named members
+	const FMemberId* ToNameIt = ToNames.GetData();
 	for (FMemberId FromName : From.GetMemberNames())
 	{
-		while (FromName != ToNames[ToIt.GetIndex()])
+		while (FromName != *ToNameIt++)
 		{
 			ToIt.SkipMember();
+			check(ToNameIt != ToNames.end());
 		}
-
+		
 		CopyMemberBinding(/* in-out */ ToIt, /* in-out */ InnerSchemaIt, /* in-out */ LoadIds,  /* out */ Footer);	
 	}
 	check(InnerSchemaIt == From.GetInnerSchemas() + From.NumInnerSchemas);
@@ -456,7 +486,7 @@ struct FLoadPlanner
 
 		if (Format == ESchemaFormat::InMemoryNames)
 		{
-			FLoadIdBinderDummy LoadIds;
+			FLoadIdBinderDummy LoadIds{Declarations.GetDebug()};
 			for (uint32 Idx = 0; Idx < NumPlans; ++Idx)
 			{
 				FLoadIdMapping Mapping;
@@ -469,7 +499,7 @@ struct FLoadPlanner
 			return CreateBatch(LoadIds);
 		}
 	
-		FLoadIdBinder LoadIds(RuntimeIds);
+		FLoadIdBinder LoadIds(RuntimeIds, Declarations.GetDebug());
 		for (uint32 Idx = 0; Idx < NumPlans; ++Idx)
 		{
 			FLoadIdMapping Mapping;
@@ -973,11 +1003,12 @@ private:
 			EMemberKind Kind = Type.GetKind();
 			if (Kind == EMemberKind::Struct)
 			{
-				(void)GrabInnerStruct(Type.AsStruct());
+				InnerStructIdx += !(Type.AsStruct().IsDynamic);
 			}
 			else if (Kind == EMemberKind::Range)
 			{
-				(void)GrabInnerRanges(Type.AsRange());
+				FMemberBindType InnermostType = GrabInnerRangeTypes(InnerRangeTypes, /* in-out */ InnerRangeIdx).Last();
+				InnerStructIdx += InnermostType.IsStruct() && !(InnermostType.AsStruct().IsDynamic);
 			}
 			++MemberIdx;
 		}
@@ -1197,6 +1228,43 @@ void ConstructAndLoadStruct(void* Dst, FStructLoadView Src)
 	ConstructAndLoadStruct(Dst, Src.Values, Src.Schema.LoadId, Src.Schema.Batch);
 }
 
+void FSchemaLoadHandle::GetInnerLoadIds(TArrayView<FOptionalSchemaId> Out) const
+{
+	const FStructSchema& ReadSchema = FStructSchemaHandle{ Batch.GetReadId(LoadId), Batch.ReadId }.Resolve();
+	TConstArrayView<FMemberType> MemberTypes = ReadSchema.GetMemberTypes();
+	TConstArrayView<FMemberType> RangeTypes = ReadSchema.GetRangeTypes();
+	check(Out.Num() == MemberTypes.Num());
+
+	FOptionalSchemaId*  OutIt = Out.GetData();
+	uint32 RangeTypeIdx = 0;
+	if (const FStructSchemaId* InnerLoadIds = Batch[LoadId].GetInnerLoadIds())
+	{
+		for (FMemberType Member : MemberTypes)
+		{
+			FMemberType Innermost = Member.IsRange() ? GrabInnerRangeTypes(RangeTypes, /* in-out */ RangeTypeIdx).Last() : Member;
+			*OutIt++ = Innermost.IsStruct() ? ToOptional(FSchemaId(*InnerLoadIds++)) : NoId;	
+		}
+	}
+	else
+	{
+		const FSchemaId* InnerSchemaIt = ReadSchema.GetInnerSchemas();
+		for (FMemberType Member : MemberTypes)
+		{
+			FMemberType Innermost = Member.IsRange() ? GrabInnerRangeTypes(RangeTypes, /* in-out */ RangeTypeIdx).Last() : Member;
+			if (Innermost.IsStruct())
+			{
+				*OutIt++ = *InnerSchemaIt++;
+			}
+			else
+			{
+				*OutIt++ = NoId;
+				InnerSchemaIt += Innermost.AsLeaf().Type == ELeafType::Enum;
+			}
+		}
+	}
+	check (OutIt == Out.end());
+}
+
 void LoadSoleStruct(void* Dst, FStructLoadView Src)
 {
 	// Todo: Optimize
diff --git a/Engine/Plugins/Experimental/PlainProps/Source/Private/PlainPropsRead.cpp b/Engine/Plugins/Experimental/PlainProps/Source/Private/PlainPropsRead.cpp
--- a/Engine/Plugins/Experimental/PlainProps/Source/Private/PlainPropsRead.cpp
+++ b/Engine/Plugins/Experimental/PlainProps/Source/Private/PlainPropsRead.cpp
@@ -250,14 +250,7 @@ inline void FMemberReader::SkipSchema(FMemberType InnermostType)
 {
 	if (InnermostType.IsStruct())
 	{
-		if (InnermostType.AsStruct().IsDynamic)
-		{
-			(void)ValueIt.Grab<uint32>();
-		}
-		else
-		{
-			++InnerSchemaIdx;
-		}
+		InnerSchemaIdx += !InnermostType.AsStruct().IsDynamic;
 	}
 	else
 	{
diff --git a/Engine/Plugins/Experimental/PlainProps/Source/Private/PlainPropsSave.cpp b/Engine/Plugins/Experimental/PlainProps/Source/Private/PlainPropsSave.cpp
--- a/Engine/Plugins/Experimental/PlainProps/Source/Private/PlainPropsSave.cpp
+++ b/Engine/Plugins/Experimental/PlainProps/Source/Private/PlainPropsSave.cpp
@@ -314,11 +314,21 @@ static void SaveMember(FMemberBuilder& Out, const void* Struct, FMemberId Name,
 	Out.AddStruct(Name, Member.Id, SaveStruct(At(Struct, Member.Offset), Member.Id, Ctx));
 }
 
-inline void SaveAllMembers(FMemberBuilder& Out, const void* Struct, const FSchemaBinding& Schema, TConstArrayView<FMemberId> MemberOrder, const FSaveContext& Ctx)
+inline void SaveAllMembers(FMemberBuilder& Out, const void* Struct, const FSchemaBinding& Schema, const FStructDeclaration& Declaration, const FSaveContext& Ctx)
 {
-	for (FMemberVisitor It(Schema); It.HasMore(); )
+	FMemberVisitor It(Schema);
+	if (Declaration.Super)
+	{
+		FStructSchemaId SuperId = It.GrabSuper();
+		checkSlow(SuperId == Declaration.Super);
+		const FSchemaBinding& SuperSchema = Ctx.Schemas.GetStruct(SuperId);
+		const FStructDeclaration& SuperDecl = Ctx.Declarations.Get(SuperId);
+		SaveAllMembers(Out, Struct, SuperSchema, SuperDecl, Ctx);
+		Out.BuildSuperStruct(Ctx.Scratch, SuperDecl, Ctx.Declarations.GetDebug());
+	}
+
+	for (FMemberId Name : Declaration.GetMemberOrder())
 	{
-		FMemberId Name = MemberOrder[It.GetIndex()];
 		switch (It.PeekKind())
 		{
 			case EMemberKind::Leaf:		SaveMember(Out, Struct, Name, Ctx, It.GrabLeaf());		break;
@@ -326,6 +336,7 @@ inline void SaveAllMembers(FMemberBuilder& Out, const void* Struct, const FSchem
 			case EMemberKind::Struct:	SaveMember(Out, Struct, Name, Ctx, It.GrabStruct());	break;
 		}
 	}
+	checkSlow(!It.HasMore());
 }
 
 FBuiltStruct* SaveStruct(const void* Struct, FStructSchemaId BindId, const FSaveContext& Ctx)
@@ -340,7 +351,7 @@ FBuiltStruct* SaveStruct(const void* Struct, FStructSchemaId BindId, const FSave
 	{
 		const FSchemaBinding& Schema = Ctx.Schemas.GetStruct(BindId);
 		Declaration = &Ctx.Declarations.Get(Schema.DeclId);	
-		SaveAllMembers(Out, Struct, Schema, Declaration->GetMemberOrder(), Ctx);
+		SaveAllMembers(Out, Struct, Schema, *Declaration, Ctx);
 	}
 
 	return Out.BuildAndReset(Ctx.Scratch, *Declaration, Ctx.Declarations.GetDebug());
@@ -425,13 +436,21 @@ static void SaveMemberDelta(FMemberBuilder& Out, const void* Struct, const void*
 static const FStructDeclaration& SaveSchemaBoundStructDelta(FMemberBuilder& Out, const void* Struct, const void* Default, const FSchemaBinding& Schema, const FSaveContext& Ctx)
 {
 	const FStructDeclaration& Declaration = Ctx.Declarations.Get(Schema.DeclId);
-	TConstArrayView<FMemberId> MemberOrder = Declaration.GetMemberOrder();
 
 	if (Declaration.Occupancy == EMemberPresence::AllowSparse)
 	{
-		for (FMemberVisitor It(Schema); It.HasMore(); )
+		FMemberVisitor It(Schema);
+		if (Declaration.Super)
+		{
+			FStructSchemaId SuperId = It.GrabSuper();
+			checkSlow(SuperId == Declaration.Super);
+			const FSchemaBinding& SuperSchema = Ctx.Schemas.GetStruct(SuperId);
+			const FStructDeclaration& SuperDecl = SaveSchemaBoundStructDelta(Out, Struct, Default, SuperSchema, Ctx);
+			Out.BuildSuperStruct(Ctx.Scratch, SuperDecl, Ctx.Declarations.GetDebug());
+		}
+
+		for (FMemberId Name : Declaration.GetMemberOrder())
 		{
-			FMemberId Name = MemberOrder[It.GetIndex()];
 			switch (It.PeekKind())
 			{
 				case EMemberKind::Leaf:		SaveMemberDelta(Out, Struct, Default, Name, Ctx, It.GrabLeaf());	break;
@@ -439,10 +458,11 @@ static const FStructDeclaration& SaveSchemaBoundStructDelta(FMemberBuilder& Out,
 				case EMemberKind::Struct:	SaveMemberDelta(Out, Struct, Default, Name, Ctx, It.GrabStruct());	break;
 			}
 		}
+		checkSlow(!It.HasMore());
 	}
 	else
 	{
-		SaveAllMembers(Out, Struct, Schema, MemberOrder, Ctx);
+		SaveAllMembers(Out, Struct, Schema, Declaration, Ctx);
 	}
 
 	return Declaration;
diff --git a/Engine/Plugins/Experimental/PlainProps/Source/Public/PlainPropsBind.h b/Engine/Plugins/Experimental/PlainProps/Source/Public/PlainPropsBind.h
--- a/Engine/Plugins/Experimental/PlainProps/Source/Public/PlainPropsBind.h
+++ b/Engine/Plugins/Experimental/PlainProps/Source/Public/PlainPropsBind.h
@@ -271,7 +271,7 @@ class FCustomBindingMap
 {
 public:
 	UE_NONCOPYABLE(FCustomBindingMap);
-	explicit FCustomBindingMap(const FDebugIds& Dbg) : Debug(Dbg) {}
+	explicit FCustomBindingMap(FDebugIds Dbg) : Debug(Dbg) {}
 	~FCustomBindingMap() { FMemory::Free(Values); }
 
 	void							Bind(FStructSchemaId BindId, ICustomBinding& Binding, const FStructDeclaration& Declaration, FInnersHandle LoweredInners);
@@ -284,7 +284,7 @@ private:
 	ICustomBinding**				Values = nullptr;
 	uint32							MaxValues = 0;
 public:
-	const FDebugIds&				Debug;
+	FDebugIds						Debug;
 };
 
 class FCustomBindings
@@ -300,7 +300,7 @@ public:
 
 	virtual FCustomBindingHandle				Find(FStructSchemaId BindId) const = 0;
 protected:
-	FCustomBindings(TArray<FInnerStruct>& Inners, const FDebugIds& Dbg) : Map(Dbg), BottomInners(Inners) {}
+	FCustomBindings(TArray<FInnerStruct>& Inners, FDebugIds Dbg) : Map(Dbg), BottomInners(Inners) {}
 	FCustomBindings(const FCustomBindings* Under) : Map(Under->Map.Debug), BottomInners(Under->BottomInners) {}
 	~FCustomBindings() = default;
 
@@ -931,6 +931,7 @@ FStructSchemaId BindCustomStructOnce()
 }
 
 inline constexpr FMemberBindType DefaultStructBindType = FMemberBindType(FStructType{EMemberKind::Struct, /* IsDynamic */ 0, /* IsSuper */ 0});
+inline constexpr FMemberBindType SuperStructBindType = FMemberBindType(FStructType{EMemberKind::Struct, /* IsDynamic */ 0, /* IsSuper */ 1});
 
 template<class Struct, class CustomBinding, class Runtime>
 FMemberBindType BindMemberStruct(FOptionalSchemaId& OutSchema)
diff --git a/Engine/Plugins/Experimental/PlainProps/Source/Public/PlainPropsLoadMember.h b/Engine/Plugins/Experimental/PlainProps/Source/Public/PlainPropsLoadMember.h
--- a/Engine/Plugins/Experimental/PlainProps/Source/Public/PlainPropsLoadMember.h
+++ b/Engine/Plugins/Experimental/PlainProps/Source/Public/PlainPropsLoadMember.h
@@ -17,6 +17,11 @@ struct FSchemaLoadHandle
 {
 	FStructSchemaId						LoadId;
 	const FLoadBatch&					Batch;
+
+	// Experimental API to help bypass FMemberLoader overhead for dense struct ranges
+	//
+	// @pre Out.Num() must equal number of members
+	PLAINPROPS_API void	GetInnerLoadIds(TArrayView<FOptionalSchemaId> Out) const;
 };
 
 // Usable via FMemberLoader or [ConstructAnd]LoadStruct()
diff --git a/Engine/Plugins/Experimental/PlainProps/Source/Public/PlainPropsRead.h b/Engine/Plugins/Experimental/PlainProps/Source/Public/PlainPropsRead.h
--- a/Engine/Plugins/Experimental/PlainProps/Source/Public/PlainPropsRead.h
+++ b/Engine/Plugins/Experimental/PlainProps/Source/Public/PlainPropsRead.h
@@ -550,9 +550,11 @@ public:
 	using SchemaType = decltype(IteratorType::Schema);
 	TStructuralRangeView(uint64 N, FMemoryView D, SchemaType S) : NumItems(N), Data(D), Schema(S) {}
 
-	uint64			Num() const		{ return NumItems; }
-	IteratorType	begin() const	{ return IteratorType(Schema, Data); }
-	IteratorType	end() const		{ return IteratorType(Schema, FMemoryView(Data.GetDataEnd(), 0)); }
+	uint64				Num() const			{ return NumItems; }
+	const SchemaType&	GetSchema() const	{ return Schema; }
+	IteratorType		begin() const		{ return IteratorType(Schema, Data); }
+	IteratorType		end() const			{ return IteratorType(Schema, FMemoryView(Data.GetDataEnd(), 0)); }
+
 
 private:
 	uint64				NumItems;
diff --git a/Engine/Plugins/Experimental/PlainProps/Source/Public/PlainPropsUeCoreBindings.h b/Engine/Plugins/Experimental/PlainProps/Source/Public/PlainPropsUeCoreBindings.h
--- a/Engine/Plugins/Experimental/PlainProps/Source/Public/PlainPropsUeCoreBindings.h
+++ b/Engine/Plugins/Experimental/PlainProps/Source/Public/PlainPropsUeCoreBindings.h
@@ -598,6 +598,7 @@ struct TBaseDeltaBinding : ICustomBinding, FBaseDeltaBinding
 		}
 	}
 
+	// Todo: Reimplement with Assign/Remove/Insert like SaveSetDelta in PlainPropsUObjectBindings.cpp
 	inline void SaveDelta(FMemberBuilder& Dst, const SetType& Src, const SetType* Default, const FSaveContext& Ctx, const TMemberSerializer<KeyType>& Keys) const
 	{
 		if (Src.IsEmpty())
@@ -927,7 +928,7 @@ struct TMapDeltaBinding : TBaseDeltaBinding<TMap<K, V, SetAllocator, KeyFuncs>,
 //		}
 //	}
 //
-//	virtual FBuiltStruct*	SaveStruct(const void* Src, const FDebugIds& Debug) const override
+//	virtual FBuiltStruct*	SaveStruct(const void* Src, FDebugIds Debug) const override
 //	{
 //		...
 //	}
diff --git a/Engine/Plugins/Experimental/PlainPropsUObject/Source/Private/PlainPropsUObjectBindings.cpp b/Engine/Plugins/Experimental/PlainPropsUObject/Source/Private/PlainPropsUObjectBindings.cpp
--- a/Engine/Plugins/Experimental/PlainPropsUObject/Source/Private/PlainPropsUObjectBindings.cpp
+++ b/Engine/Plugins/Experimental/PlainPropsUObject/Source/Private/PlainPropsUObjectBindings.cpp
@@ -313,8 +313,9 @@ FCommonStructIds::FCommonStructIds(const FCommonScopeIds& Scopes, TIdIndexer<FNa
 FCommonMemberIds::FCommonMemberIds(TIdIndexer<FName>& Names)
 : Key(						Names.NameMember("Key"))
 , Value(					Names.NameMember("Value"))
-, Add(						Names.NameMember("Add"))
-, Del(						Names.NameMember("Del"))
+, Assign(					Names.NameMember("Assign"))
+, Remove(					Names.NameMember("Remove"))
+, Insert(					Names.NameMember("Insert"))
 , Id(						Names.NameMember("Id"))
 , Object(					Names.NameMember("Object"))
 , Function(					Names.NameMember("Function"))
@@ -329,7 +330,6 @@ FGlobals::FGlobals()
 : Types(FDebugIds(Names))
 , Schemas(FDebugIds(Names))
 , Customs(FDebugIds(Names))
-, DeltaCustoms(Customs)
 , Scopes(Names)
 , Structs(Scopes, Names)
 , Typenames(Names)
@@ -344,7 +344,7 @@ FGlobals GUE;
 static constexpr uint64 LeafMask =			CASTCLASS_FNumericProperty | CASTCLASS_FEnumProperty| CASTCLASS_FBoolProperty;
 static constexpr uint64 IntSMask =			CASTCLASS_FInt8Property | CASTCLASS_FInt16Property | CASTCLASS_FIntProperty | CASTCLASS_FInt64Property;
 static constexpr uint64 IntUMask =			CASTCLASS_FByteProperty | CASTCLASS_FUInt16Property | CASTCLASS_FUInt32Property | CASTCLASS_FUInt64Property;
-static constexpr uint64 ContainerMask =		CASTCLASS_FArrayProperty | CASTCLASS_FSetProperty| CASTCLASS_FMapProperty  | CASTCLASS_FOptionalProperty;
+static constexpr uint64 ContainerMask =		CASTCLASS_FArrayProperty | CASTCLASS_FSetProperty | CASTCLASS_FMapProperty  | CASTCLASS_FOptionalProperty;
 static constexpr uint64 StringMask =		CASTCLASS_FStrProperty | CASTCLASS_FUtf8StrProperty | CASTCLASS_FAnsiStrProperty | CASTCLASS_FVerseStringProperty;
 static constexpr uint64 CommonStructMask =	CASTCLASS_FNameProperty | CASTCLASS_FTextProperty |  CASTCLASS_FFieldPathProperty |
 											CASTCLASS_FObjectProperty | CASTCLASS_FWeakObjectProperty | CASTCLASS_FSoftObjectProperty | CASTCLASS_FLazyObjectProperty |
@@ -497,10 +497,8 @@ static void DeclareSuperMembers(FMemberArray& Out, const UStruct* Struct)
 	}
 }
 
-static FStructSchemaId DeclareStruct(const UStruct* Struct)
+static void DeclareStruct(const UStruct* Struct, FTypeId Type, FStructSchemaId Id)
 {
-	FTypeId Type = IndexType(Struct);
-	FStructSchemaId Id = GUE.Names.IndexStruct(Type);
 	FOptionalStructSchemaId  Super = IndexSuper(Struct);
 	EMemberPresence Occupancy = GetOccupancy(Struct);
 	
@@ -514,7 +512,13 @@ static FStructSchemaId DeclareStruct(const UStruct* Struct)
 	DeclareMembers(/* out */ Members, Struct);
 
 	GUE.Types.DeclareStruct(Id, Type, Members, Occupancy, Super);
+}
 
+static FStructSchemaId DeclareStruct(const UStruct* Struct)
+{
+	FTypeId Type = IndexType(Struct);
+	FStructSchemaId Id = GUE.Names.IndexStruct(Type);
+	DeclareStruct(Struct, Type, Id);
 	return Id;
 }
 
@@ -1600,12 +1604,12 @@ struct TLeafRangeSerializer
 		return Cast(In);
 	}
 
-	inline static void LoadItem(void* Dst, FByteReader& SrcBytes, FBitCacheReader& SrcBits, const FLoadBatch&) requires (std::is_same_v<EquivalentType, bool>)
+	inline static void LoadItem(void* Dst, FByteReader& SrcBytes, FBitCacheReader& SrcBits, FOptionalSchemaId, const FLoadBatch&) requires (std::is_same_v<EquivalentType, bool>)
 	{
 		*static_cast<bool*>(Dst) = SrcBits.GrabNext(SrcBytes);
 	}
 
-	inline static void LoadItem(void* Dst, FByteReader& SrcBytes, FBitCacheReader&, const FLoadBatch&)
+	inline static void LoadItem(void* Dst, FByteReader& SrcBytes, FBitCacheReader&, FOptionalSchemaId, const FLoadBatch&)
 	{
 		*static_cast<EquivalentType*>(Dst) = Cast(SrcBytes.GrabBytes(Size));
 	}
@@ -1616,28 +1620,28 @@ struct FStructRangeSerializer
 	using RangeSaver = FStructRangeSaver;
 	
 	FMemberType				InnerType;
-	FStructSchemaId			Id;
+	FStructSchemaId			SaveId;
 
 	explicit FStructRangeSerializer(FMemberBinding Item)
 	: InnerType(Item.InnermostType.AsStruct())
-	, Id(static_cast<FStructSchemaId>(Item.InnermostSchema.Get()))
+	, SaveId(static_cast<FStructSchemaId>(Item.InnermostSchema.Get()))
 	{
 		check(Item.RangeBindings.IsEmpty());
 	}
 
 	inline FMemberSchema MakeMemberSchema() const
 	{
-		return { FMemberType(DefaultRangeMax), InnerType, 1, FOptionalSchemaId(Id), nullptr };
+		return { FMemberType(DefaultRangeMax), InnerType, 1, FOptionalSchemaId(SaveId), nullptr };
 	}
 
 	inline FBuiltStruct* SaveItem(const void* In, const FSaveContext& Ctx) const
 	{
-		return SaveStruct(In, Id, Ctx);
+		return SaveStruct(In, SaveId, Ctx);
 	}
 
-	inline void LoadItem(void* Dst, FByteReader& SrcBytes, FBitCacheReader&, const FLoadBatch& Batch) const
+	void LoadItem(void* Dst, FByteReader& SrcBytes, FBitCacheReader&, FOptionalSchemaId LoadId, const FLoadBatch& Batch) const
 	{
-		LoadStruct(Dst, FByteReader(SrcBytes.GrabSkippableSlice()), Id, Batch);
+		LoadStruct(Dst, FByteReader(SrcBytes.GrabSkippableSlice()), static_cast<FStructSchemaId>(LoadId.Get()), Batch);
 	}
 };
 
@@ -1645,14 +1649,14 @@ struct FNestedRangeSerializer
 {
 	using RangeSaver = FNestedRangeSaver;
 
-	FOptionalSchemaId								InnermostId;
+	FOptionalSchemaId								InnermostSaveId;
 	uint16											NumInners;
 	TArray<FMemberType, TInlineAllocator<8>>		InnerTypes;
 	TArray<FMemberBindType, TInlineAllocator<8>>	InnerBindTypes;
 	TArray<FRangeBinding, TInlineAllocator<2>>		InnerBindings;
 
 	explicit FNestedRangeSerializer(FMemberBinding Item)
-	: InnermostId(Item.InnermostSchema)
+	: InnermostSaveId(Item.InnermostSchema)
 	, NumInners(IntCastChecked<uint16>(1 + Item.RangeBindings.Num()))
 	, InnerBindings(Item.RangeBindings)
 	{
@@ -1669,18 +1673,18 @@ struct FNestedRangeSerializer
 
 	inline FMemberSchema MakeMemberSchema() const
 	{
-		return { FMemberType(DefaultRangeMax), InnerTypes[0], NumInners, InnermostId, InnerTypes.GetData() };
+		return { FMemberType(DefaultRangeMax), InnerTypes[0], NumInners, InnermostSaveId, InnerTypes.GetData() };
 	}
 
 	FBuiltRange* SaveItem(const void* In, const FSaveContext& Ctx) const
 	{
-		FRangeMemberBinding Member = { InnerBindTypes.GetData() + 1, InnerBindings.GetData(), NumInners - 1, InnermostId, 0 };
+		FRangeMemberBinding Member = { InnerBindTypes.GetData() + 1, InnerBindings.GetData(), NumInners - 1, InnermostSaveId, 0 };
 		return SaveRange(In, Member, Ctx);
 	}
 
-	inline void LoadItem(void* Dst, FByteReader& SrcBytes, FBitCacheReader& SrcBits, const FLoadBatch& Batch) const
+	inline void LoadItem(void* Dst, FByteReader& SrcBytes, FBitCacheReader& SrcBits, FOptionalSchemaId InnermostLoadId, const FLoadBatch& Batch) const
 	{
-		FRangeLoadSchema Schema = {InnerTypes[0], InnermostId, InnerTypes.GetData() + 1, Batch};
+		FRangeLoadSchema Schema = { InnerTypes[1], InnermostLoadId, InnerTypes.GetData() + 2, Batch};
 		LoadRange(Dst, SrcBytes, SrcBits, DefaultRangeMax, Schema, InnerBindings);
 	}
 };
@@ -1787,14 +1791,14 @@ struct TSubSetIterator
 	uint32 CountNum() const			{ return Subset.CountSetBits(); }
 };
 
-using FSubSetIterator = TSubSetIterator<FScriptSet, FScriptSetLayout>;
-using FSubMapIterator = TSubSetIterator<FScriptMap, FScriptMapLayout>;
-
 // Helps save TSet/TMap deltas
 template<class ScriptType, class LayoutType, class SerializerType>
 FTypedRange SaveAll(const ScriptType& Set, const LayoutType& Layout, const SerializerType& Serializer, const FSaveContext& Ctx)
 {
-	check(!Set.IsEmpty());
+	if (Set.IsEmpty())
+	{
+		return { Serializer.MakeMemberSchema(), nullptr };
+	}
 
 	typename SerializerType::RangeSaver Range(Ctx.Scratch, static_cast<uint64>(Set.Num()));
 	for (int32 Idx = 0, Max = Set.GetMaxIndex(); Idx < Max; ++Idx)
@@ -1820,33 +1824,88 @@ FTypedRange SaveSome(SubSetIteratorType& It, const SerializerType& Serializer, c
 }
 
 template<typename BindingType, typename ScriptType>
-void LoadSetDelta(const BindingType& Binding, ScriptType& Dst, FStructLoadView Src, ECustomLoadMethod Method)
+void SaveSetDelta(const BindingType& Binding, FMemberBuilder& Dst, const ScriptType& Src, const ScriptType* Default, const FSaveContext& Ctx)
 {
-	check(Method == ECustomLoadMethod::Assign);
-	FMemberLoader Members(Src);				
-	while (Members.HasMore())
+	if (!Default)
 	{
-		if (Members.PeekName() == GUE.Members.Add)
+		Dst.AddRange(GUE.Members.Assign, SaveAll(Src, Binding.Layout, Binding.GetItemRange(), Ctx));
+	}
+	else if (Default->IsEmpty())
+	{
+		if (!Src.IsEmpty())
 		{
-			if (Dst.IsEmpty())
-			{
-				Binding.Append(Dst, Members.GrabRange());
-			}
-			else
-			{
-				Binding.Insert(Dst, Members.GrabRange());
-			}
-			check(!Members.HasMore());
-			break;
+			Dst.AddRange(GUE.Members.Insert, SaveAll(Src, Binding.Layout, Binding.GetItemRange(), Ctx));
+		}
+	}
+	else if (Src.IsEmpty())
+	{
+		Dst.AddRange(GUE.Members.Remove, SaveAll(*Default, Binding.Layout, Binding.GetKeyRange(), Ctx));
+	}
+	else // Neither are empty
+	{
+		TBitArray<> RemoveIds(false, Default->GetMaxIndex());
+		for (int32 Idx = 0, Max = Default->GetMaxIndex(); Idx < Max; ++Idx)
+		{
+			RemoveIds[Idx] = Default->IsValidIndex(Idx) && !Binding.HasKey(Src, Default->GetData(Idx, Binding.Layout));
+		}
+		if (typename BindingType::SubSetIterator Removed{Binding.Layout, *Default, RemoveIds})
+		{
+			Dst.AddRange(GUE.Members.Remove, SaveSome(Removed, Binding.GetKeyRange(), Ctx));
+		}
+			
+		TBitArray<> InsertIds(false, Src.GetMaxIndex());
+		for (int32 Idx = 0, Max = Src.GetMaxIndex(); Idx < Max; ++Idx)
+		{
+			InsertIds[Idx] = Src.IsValidIndex(Idx) && !Binding.HasItem(*Default, Src.GetData(Idx, Binding.Layout));
+		}
+		if (typename BindingType::SubSetIterator Inserted{Binding.Layout, Src, InsertIds})
+		{
+			Dst.AddRange(GUE.Members.Insert, SaveSome(Inserted, Binding.GetItemRange(), Ctx));
 		}
+	}
+}
 
-		check(Members.PeekName() == GUE.Members.Del);
-		FRangeLoadView Keys = Members.GrabRange();
-		if (!Dst.IsEmpty())
+template<typename BindingType, typename ScriptType>
+void InsertSetItems(const BindingType& Binding, ScriptType& Dst, FRangeLoadView Items)
+{
+	// Insert
+	if (Dst.IsEmpty())
+	{
+		Binding.AssignEmpty(Dst, Items);
+	}
+	else
+	{
+		Binding.InsertNonEmpty(Dst, Items);
+	}
+}
+
+template<typename BindingType, typename ScriptType>
+void LoadSetDelta(const BindingType& Binding, ScriptType& Dst, FStructLoadView Src)
+{
+	FMemberLoader Members(Src);
+	FOptionalMemberId Name = Members.PeekName();
+	FRangeLoadView Range = Members.GrabRange();
+	if (Name == GUE.Members.Insert)
+	{
+		InsertSetItems(Binding, Dst, Range);
+	} 
+	else if (Name == GUE.Members.Assign)
+	{
+		Binding.DestroyAll(Dst);
+		Binding.AssignEmpty(Dst, Range);
+	}
+	else
+	{
+		checkSlow(Name == GUE.Members.Remove);
+		Binding.Remove(Dst, Range);
+		if (Members.HasMore())
 		{
-			Binding.Remove(Dst, Keys);
+			checkSlow(Members.PeekNameUnchecked() == GUE.Members.Insert);
+			InsertSetItems(Binding, Dst, Members.GrabRange());
 		}
-	}	
+	}
+
+	checkSlow(!Members.HasMore());
 }
 
 template<typename BindingType, typename ScriptType>
@@ -1861,7 +1920,7 @@ inline bool DiffSet(const BindingType& Binding, const ScriptType& A, const Scrip
 	{
 		for (int32 IdxA = 0, MaxA = A.GetMaxIndex(); IdxA < MaxA; ++IdxA)
 		{
-			if (A.IsValidIndex(IdxA) && !Binding.HasElem(B, A.GetData(IdxA, Binding.Layout)))
+			if (A.IsValidIndex(IdxA) && !Binding.HasItem(B, A.GetData(IdxA, Binding.Layout)))
 			{
 				return true;
 			}
@@ -1878,12 +1937,16 @@ struct TSetPropertyBinding : IItemRangeBinding, ICustomBinding
 {
 	using SetType = FScriptSet;
 	using LeafType = EquivalentLeafType<ElemKind>;
+	using SubSetIterator = TSubSetIterator<FScriptSet, FScriptSetLayout>;
 	static constexpr bool bLeaves = !std::is_void_v<LeafType>;
 
 	const FScriptSetLayout							Layout;
 	const TInnerProperty<ElemKind>					Inner;
 	const TPropertyRangeSerializer<ElemKind>		Range;
 
+	const TPropertyRangeSerializer<ElemKind>& GetKeyRange() const { return Range; }
+	const TPropertyRangeSerializer<ElemKind>& GetItemRange() const { return Range; }
+
 	TSetPropertyBinding(FSetProperty* In, FMemberBinding Elem)
 	: IItemRangeBinding(GUE.Typenames.Set)
 	, Layout(In->SetLayout)
@@ -1906,19 +1969,31 @@ struct TSetPropertyBinding : IItemRangeBinding, ICustomBinding
 	{
 		if (int32 Idx = FindIndex(Set, Elem); Idx != INDEX_NONE)
 		{
+			DestroyElem(Set, Idx);
 			Set.RemoveAt(Idx, Layout);
 		}
 	}
-
-	inline bool HasElem(const FScriptSet& Set, const void* Elem) const
+	
+	inline bool HasItem(const FScriptSet& Set, const void* Elem) const
 	{
 		return FindIndex(Set, Elem) != INDEX_NONE;
 	}
 
+	inline bool HasKey(const FScriptSet& Set, const void* Elem) const
+	{
+		return HasItem(Set, Elem);
+	}
+
 	inline void Rehash(FScriptSet& Set) const
 	{
 		Set.Rehash(Layout, MakeHashFn(Inner));
 	}
+
+	inline void DestroyElem(FScriptSet& Set) const requires (bLeaves) {}
+	inline void DestroyElem(FScriptSet& Set, int32 Idx) const
+	{
+		Inner.DestroyItem(Set.GetData(Idx, Layout));
+	}
 	
 	inline void DestroyAll(FScriptSet& Set) const requires (bLeaves) {}
 	inline void DestroyAll(FScriptSet& Set) const
@@ -1970,50 +2045,16 @@ struct TSetPropertyBinding : IItemRangeBinding, ICustomBinding
 			Rehash(Set);
 		}
 	}
-	
-	virtual void SaveCustom(FMemberBuilder& Dst, const void* InSrc, const void* InDefault, const FSaveContext& Ctx) override
-	{
-		const FScriptSet& Src = *static_cast<const FScriptSet*>(InSrc);
-		const FScriptSet* Default = static_cast<const FScriptSet*>(InDefault);
 
-		if (Src.IsEmpty())
-		{
-			if (Default && !Default->IsEmpty())
-			{
-				Dst.AddRange(GUE.Members.Del, SaveAll(*Default, Layout, Range, Ctx));
-			}
-		}
-		else if (Default && !Default->IsEmpty())
-		{
-			TBitArray<> DeletedIds(false, Default->GetMaxIndex());
-			for (int32 Idx = 0, Max = Default->GetMaxIndex(); Idx < Max; ++Idx)
-			{
-				DeletedIds[Idx] = Default->IsValidIndex(Idx) && HasElem(Src, Default->GetData(Idx, Layout));
-			}
-			if (FSubSetIterator Deleted{Layout, *Default, DeletedIds})
-			{
-				Dst.AddRange(GUE.Members.Del, SaveSome(Deleted, Range, Ctx));
-			}
-			
-			TBitArray<> AddedIds(false, Src.GetMaxIndex());
-			for (int32 Idx = 0, Max = Src.GetMaxIndex(); Idx < Max; ++Idx)
-			{
-				AddedIds[Idx] = Src.IsValidIndex(Idx) && !HasElem(*Default, Src.GetData(Idx, Layout));
-			}
-			if (FSubSetIterator Added{Layout, Src, AddedIds})
-			{
-				Dst.AddRange(GUE.Members.Add, SaveSome(Added, Range, Ctx));
-			}
-		}
-		else
-		{
-			Dst.AddRange(GUE.Members.Add, SaveAll(Src, Layout, Range, Ctx));
-		}
+	virtual void SaveCustom(FMemberBuilder& Dst, const void* Src, const void* Default, const FSaveContext& Ctx) override
+	{
+		SaveSetDelta(*this, Dst, *static_cast<const FScriptSet*>(Src), static_cast<const FScriptSet*>(Default), Ctx);
 	}
 
 	virtual void LoadCustom(void* Dst, FStructLoadView Src, ECustomLoadMethod Method) const override
 	{
-		LoadSetDelta(*this, *static_cast<FScriptSet*>(Dst), Src, Method);
+		check(Method == ECustomLoadMethod::Assign);
+		LoadSetDelta(*this, *static_cast<FScriptSet*>(Dst), Src);
 	}
 
 	virtual bool DiffCustom(const void* A, const void* B) const override
@@ -2022,7 +2063,7 @@ struct TSetPropertyBinding : IItemRangeBinding, ICustomBinding
 	}
 
 	// Load into empty set
-	inline void Append(FScriptSet& Dst, FRangeLoadView Src) const
+	inline void AssignEmpty(FScriptSet& Dst, FRangeLoadView Src) const
 	{
 		SIZE_T Stride = GetStride();
 		uint8* It = SetNumUninitialized(Dst, Layout, Src.Num());
@@ -2058,11 +2099,11 @@ struct TSetPropertyBinding : IItemRangeBinding, ICustomBinding
 	}
 
 	// Load leaves into non-empty set
-	inline void Insert(FScriptSet& Dst, FRangeLoadView Src) const requires(bLeaves)
+	inline void InsertNonEmpty(FScriptSet& Dst, FRangeLoadView Src) const requires(bLeaves)
 	{	
 		for (LeafType Item : CastAs<LeafType>(Src.AsLeaves()))
 		{
-			if (!HasElem(Dst, &Item))
+			if (!HasItem(Dst, &Item))
 			{
 				 void* Elem = Dst.GetData(Dst.AddUninitialized(Layout), Layout);
 				 *static_cast<LeafType*>(Elem) = Item;
@@ -2081,7 +2122,7 @@ struct TSetPropertyBinding : IItemRangeBinding, ICustomBinding
 	}
 
 	// Load structs or ranges into non-empty set
-	inline void Insert(FScriptSet& DstSet, FRangeLoadView Src) const
+	inline void InsertNonEmpty(FScriptSet& DstSet, FRangeLoadView Src) const
 	{
 		// Written to avoid FProperty::CopyCompleteValue_InContainer dependency
 		// Items are loaded directly into sparse array and then removed if a duplicate existed
@@ -2089,15 +2130,13 @@ struct TSetPropertyBinding : IItemRangeBinding, ICustomBinding
 		const int32 OldNum = Dst.NumUnchecked();
 		int32 TmpIdx;
 		void* Tmp = AddItem(Dst, /* out */ TmpIdx);
-		bool bSkip = true;
 		if constexpr (ElemKind == EPropertyKind::Range)
 		{
 			TConstArrayView<FRangeBinding> InnerBindings = Range.InnerBindings;
 			for (FRangeLoadView Item : Src.AsRanges())
 			{
 				LoadRange(Tmp, Item, InnerBindings);
-				bSkip = HasElem(DstSet, Tmp);
-				Tmp = bSkip ? Tmp : AddItem(Dst, /* out */ TmpIdx);
+				Tmp = HasItem(DstSet, Tmp) ? Tmp : AddItem(Dst, /* out */ TmpIdx);
 			}
 		}
 		else
@@ -2105,16 +2144,12 @@ struct TSetPropertyBinding : IItemRangeBinding, ICustomBinding
 			for (FStructLoadView Item : Src.AsStructs())
 			{
 				LoadStruct(Tmp, Item);
-				bSkip = HasElem(DstSet, Tmp);
-				Tmp = bSkip ? Tmp : AddItem(Dst, /* out */ TmpIdx);
+				Tmp = HasItem(DstSet, Tmp) ? Tmp : AddItem(Dst, /* out */ TmpIdx);
 			}
 		}
 
-		if (bSkip)
-		{
-			Inner.DestroyItem(Tmp);
-			Dst.RemoveAtUninitialized(Layout.SparseArrayLayout, TmpIdx, 1);
-		}
+		Inner.DestroyItem(Tmp);
+		Dst.RemoveAtUninitialized(Layout.SparseArrayLayout, TmpIdx, 1);
 
 		if (Dst.NumUnchecked() != OldNum)
 		{
@@ -2157,47 +2192,54 @@ struct TSetPropertyBinding : IItemRangeBinding, ICustomBinding
 
 		Inner.DestroyItem(Tmp);
 	}
+
+	inline void DestroyRemoved(FScriptSet& Dst, int32 Idx) const
+	{
+		checkSlow(Idx < Dst.GetMaxIndex());
+		checkSlow(!Dst.IsValidIndex(Idx));
+		void* Elem = AsSparseArray(Dst).GetData(Idx, Layout.SparseArrayLayout);
+		Inner.DestroyItem(Elem);
+	}
 };
 
 ////////////////////////////////////////////////////////////////////////////////////////////////
 
 class FSetBindings
 {
-	TMap<FParameterBinding, FRangeBinding> RangeBindings;
+	TMap<FParameterBinding, FStructSchemaId> Bindings;
 
 	template<EPropertyKind ItemKind>
-	IItemRangeBinding* BindNew(FSetProperty* Property, FMemberBinding Elem, FStructSchemaId BindId, const FStructDeclaration& Declaration) 
+	void BindNew(FSetProperty* Property, FMemberBinding Elem, FStructSchemaId BindId, const FStructDeclaration& Declaration) 
 	{
 		// Todo: Ownership / memory leak
-		TSetPropertyBinding<ItemKind>* Out = new TSetPropertyBinding<ItemKind>(Property, Elem);
-		GUE.DeltaCustoms.BindStruct(BindId, *Out, Declaration, {});
-		return Out;
+		TSetPropertyBinding<ItemKind>* Leak = new TSetPropertyBinding<ItemKind>(Property, Elem);
+		GUE.Customs.BindStruct(BindId, *Leak, Declaration, {});
 	}
 
-	IItemRangeBinding* DeltaBindNew(FSetProperty* Property, FMemberBinding Elem, FStructSchemaId BindId, const FStructDeclaration& Declaration)
+	void DeltaBindNew(FSetProperty* Property, FMemberBinding Elem, FStructSchemaId BindId, const FStructDeclaration& Declaration)
 	{
 		switch (GetPropertyKind(Elem))
 		{
-		case EPropertyKind::Range:		return BindNew<EPropertyKind::Range	>(Property, Elem, BindId, Declaration);
-		case EPropertyKind::Struct:		return BindNew<EPropertyKind::Struct>(Property, Elem, BindId, Declaration);
-		case EPropertyKind::Bool:		return BindNew<EPropertyKind::Bool	>(Property, Elem, BindId, Declaration);
-		case EPropertyKind::U8:			return BindNew<EPropertyKind::U8	>(Property, Elem, BindId, Declaration);
-		case EPropertyKind::U16:		return BindNew<EPropertyKind::U16	>(Property, Elem, BindId, Declaration);
-		case EPropertyKind::U32:		return BindNew<EPropertyKind::U32	>(Property, Elem, BindId, Declaration);
-		case EPropertyKind::U64:		return BindNew<EPropertyKind::U64	>(Property, Elem, BindId, Declaration);
-		case EPropertyKind::F32:		return BindNew<EPropertyKind::F32	>(Property, Elem, BindId, Declaration);
-		case EPropertyKind::F64:		return BindNew<EPropertyKind::F64	>(Property, Elem, BindId, Declaration);
-		default: check(false);			return nullptr;
+		case EPropertyKind::Range:		BindNew<EPropertyKind::Range >(Property, Elem, BindId, Declaration); break;
+		case EPropertyKind::Struct:		BindNew<EPropertyKind::Struct>(Property, Elem, BindId, Declaration); break;
+		case EPropertyKind::Bool:		BindNew<EPropertyKind::Bool	 >(Property, Elem, BindId, Declaration); break;
+		case EPropertyKind::U8:			BindNew<EPropertyKind::U8	 >(Property, Elem, BindId, Declaration); break;
+		case EPropertyKind::U16:		BindNew<EPropertyKind::U16	 >(Property, Elem, BindId, Declaration); break;
+		case EPropertyKind::U32:		BindNew<EPropertyKind::U32	 >(Property, Elem, BindId, Declaration); break;
+		case EPropertyKind::U64:		BindNew<EPropertyKind::U64	 >(Property, Elem, BindId, Declaration); break;
+		case EPropertyKind::F32:		BindNew<EPropertyKind::F32	 >(Property, Elem, BindId, Declaration); break;
+		case EPropertyKind::F64:		BindNew<EPropertyKind::F64	 >(Property, Elem, BindId, Declaration); break;
+		default:						check(false); break;
 		}
 	}
 public:
 
-	FRangeBinding Bind(FSetProperty* Property, FMemberBinding Elem)
+	FStructSchemaId Bind(FSetProperty* Property, FMemberBinding Elem)
 	{
 		check(Elem.Offset == 0);
-		if (const FRangeBinding* RangeBinding = RangeBindings.Find(FParameterBinding(Elem)))
+		if (const FStructSchemaId* BindId = Bindings.Find(FParameterBinding(Elem)))
 		{
-			return *RangeBinding;
+			return *BindId;
 		}
 
 		// Index custom delta binding struct name
@@ -2208,12 +2250,12 @@ public:
 						 : BindType;
 		FStructSchemaId BindId = GUE.Names.IndexStruct(BindType);
 		FStructSchemaId DeclId = Param.IsLowered() ? GUE.Names.IndexStruct(DeclType) : BindId;
-		FMemberId Members[2] = { GUE.Members.Del, GUE.Members.Add };
+		FMemberId Members[] = { GUE.Members.Assign, GUE.Members.Remove, GUE.Members.Insert };
 
 		// Todo: Ownership / memory leak
 		const FStructDeclaration& Declaration = GUE.Types.DeclareStruct(DeclId, DeclType, Members, EMemberPresence::AllowSparse);
-		IItemRangeBinding* Out = DeltaBindNew(Property, Elem, BindId, Declaration);
-		return RangeBindings.Emplace(Elem, FRangeBinding(*Out, DefaultRangeMax));
+		DeltaBindNew(Property, Elem, BindId, Declaration);
+		return Bindings.Emplace(Elem, BindId);
 	}
 };
 static FSetBindings GSets;
@@ -2222,13 +2264,13 @@ static FSetBindings GSets;
 
 // Flat TMap binding
 template<class ScriptMap, EPropertyKind KeyKind, EPropertyKind ValueKind>
-struct TMapPropertyBinding : IItemRangeBinding
+struct TMapPropertyItemBinding : IItemRangeBinding
 {
 	const FScriptMapLayout							Layout;
 	const TInnerProperty<KeyKind>					InnerKey;
 	const TInnerProperty<ValueKind>					InnerValue;
 	
-	TMapPropertyBinding(FMapProperty* In)
+	TMapPropertyItemBinding(FMapProperty* In)
 	: IItemRangeBinding(GUE.Typenames.Map)
 	, Layout(In->MapLayout)
 	, InnerKey(In->KeyProp)
@@ -2317,21 +2359,25 @@ struct FMapMemberBindings
 };
 
 template <EPropertyKind KeyKind, EPropertyKind ValueKind>
-struct TMapPropertyDeltaBinding : TMapPropertyBinding<FScriptMap, KeyKind, ValueKind>, ICustomBinding
+struct TMapPropertyCustomBinding : TMapPropertyItemBinding<FScriptMap, KeyKind, ValueKind>, ICustomBinding
 {
-	using Super = TMapPropertyBinding<FScriptMap, KeyKind, ValueKind>;
+	using Super = TMapPropertyItemBinding<FScriptMap, KeyKind, ValueKind>;
 	using Super::Layout;
 	using Super::InnerKey;
 	using Super::InnerValue;
 	using Super::GetStride;
 	using Super::InitMap;
 	using Super::Rehash;
+	using SubSetIterator = TSubSetIterator<FScriptMap, FScriptMapLayout>;
 	
 	const TPropertyRangeSerializer<KeyKind>			KeyRange;
 	const TPropertyRangeSerializer<ValueKind>		ValueRange;
 	const FStructRangeSerializer					PairRange;
 
-	TMapPropertyDeltaBinding(FMapProperty* Map, FMapMemberBindings Members)
+	const TPropertyRangeSerializer<KeyKind>&		GetKeyRange() const { return KeyRange; }
+	const FStructRangeSerializer&					GetItemRange() const { return PairRange; }
+
+	TMapPropertyCustomBinding(FMapProperty* Map, FMapMemberBindings Members)
 	: Super(Map)
 	, KeyRange(Members.Key)
 	, ValueRange(Members.Value)
@@ -2353,7 +2399,7 @@ struct TMapPropertyDeltaBinding : TMapPropertyBinding<FScriptMap, KeyKind, Value
 		return FindKey(Map, Key) != INDEX_NONE;
 	}
 
-	inline bool HasPair(const FScriptMap& Map, const void* Pair) const
+	inline bool HasItem(const FScriptMap& Map, const void* Pair) const
 	{
 		const void* Key = Pair;
 		if (int32 Idx = FindKey(Map, Key); Idx != INDEX_NONE)
@@ -2364,11 +2410,6 @@ struct TMapPropertyDeltaBinding : TMapPropertyBinding<FScriptMap, KeyKind, Value
 		return false;
 	}
 
-	inline bool HasElem(const FScriptMap& Map, const void* Pair) const
-	{
-		return HasPair(Map, Pair);
-	}
-
 	inline uint8* AddPair(FScriptSparseArray& Dst, int32& OutIdx, int32 ValueOffset) const
 	{
 		OutIdx = Dst.AddUninitialized(Layout.SetLayout.SparseArrayLayout);
@@ -2378,81 +2419,64 @@ struct TMapPropertyDeltaBinding : TMapPropertyBinding<FScriptMap, KeyKind, Value
 		return Out;
 	}
 
+	inline void DestroyPair(FScriptMap& Map, int32 Idx) const
+	{
+		if (InnerKey.bDestruct || InnerValue.bDestruct)
+		{
+			void* Pair = Map.GetData(Idx, Layout);
+			InnerKey.DestroyItem(Pair);
+			InnerValue.DestroyItem(static_cast<uint8*>(Pair) + Layout.ValueOffset);
+		}
+	}
+
 	inline void RemoveKey(FScriptMap& Map, const void* Key) const
 	{
 		if (int32 Idx = FindKey(Map, Key); Idx != INDEX_NONE)
 		{
+			DestroyPair(Map, Idx);
 			Map.RemoveAt(Idx, Layout);
 		}
 	}
 
-	virtual void SaveCustom(FMemberBuilder& Dst, const void* InSrc, const void* InDefault, const FSaveContext& Ctx) override
+	virtual void SaveCustom(FMemberBuilder& Dst, const void* Src, const void* Default, const FSaveContext& Ctx) override
 	{
-		const FScriptMap& Src = *static_cast<const FScriptMap*>(InSrc);
-		const FScriptMap* Default = static_cast<const FScriptMap*>(InDefault);
-		
-		if (Src.IsEmpty())
-		{
-			if (Default && !Default->IsEmpty())
-			{
-				Dst.AddRange(GUE.Members.Del, SaveAll(*Default, Layout, KeyRange, Ctx));
-			}
-		}
-		else if (Default && !Default->IsEmpty())
-		{
-			TBitArray<> DeletedIds(false, Default->GetMaxIndex());
-			for (int32 Idx = 0, Max = Default->GetMaxIndex(); Idx < Max; ++Idx)
-			{
-				DeletedIds[Idx] = Default->IsValidIndex(Idx) && HasKey(Src, Default->GetData(Idx, Layout));
-			}
-			if (FSubMapIterator Deleted{Layout, *Default, DeletedIds})
-			{
-				Dst.AddRange(GUE.Members.Del, SaveSome(Deleted, KeyRange, Ctx));
-			}
-			
-			TBitArray<> AddedIds(false, Src.GetMaxIndex());
-			for (int32 Idx = 0, Max = Src.GetMaxIndex(); Idx < Max; ++Idx)
-			{
-				AddedIds[Idx] = Src.IsValidIndex(Idx) && !HasPair(*Default, Src.GetData(Idx, Layout));
-			}
-			if (FSubMapIterator Added{Layout, Src, AddedIds})
-			{
-				Dst.AddRange(GUE.Members.Add, SaveSome(Added, PairRange, Ctx));
-			}
-		}
-		else
-		{
-			Dst.AddRange(GUE.Members.Add, SaveAll(Src, Layout, PairRange, Ctx));
-		}
+		SaveSetDelta(*this, Dst, *static_cast<const FScriptMap*>(Src), static_cast<const FScriptMap*>(Default), Ctx);
 	}
 
 	virtual void LoadCustom(void* Dst, FStructLoadView Src, ECustomLoadMethod Method) const override
 	{
-		LoadSetDelta(*this, *static_cast<FScriptMap*>(Dst), Src, Method);
+		check(Method == ECustomLoadMethod::Assign);
+		LoadSetDelta(*this, *static_cast<FScriptMap*>(Dst), Src);
 	}
 
 	// Load into empty set
-	inline void Append(FScriptMap& Dst, FRangeLoadView Src) const
+	inline void AssignEmpty(FScriptMap& Dst, FRangeLoadView Src) const
 	{
 		const int32 ValueOffset = Layout.ValueOffset;
 		const SIZE_T Stride = GetStride();
 
 		uint8* It = InitMap(Dst, Src.Num());
-		for (FStructLoadView Struct : Src.AsStructs())
+		
+		if (!Src.IsEmpty())
 		{
-			// Equivalent to LoadStruct(It, Struct);
-			FBitCacheReader Bits;
-			KeyRange.LoadItem(It, /* in-out */ Struct.Values, /* in-out */ Bits, Struct.Schema.Batch);
-			ValueRange.LoadItem(It + ValueOffset, /* in-out */ Struct.Values, /* in-out */ Bits, Struct.Schema.Batch);
-
-			It += Stride;
+			FOptionalSchemaId InnerLoadIds[2];
+			FStructRangeLoadView Structs = Src.AsStructs();
+			Structs.GetSchema().GetInnerLoadIds(/* out */ MakeArrayView(InnerLoadIds));
+			for (FStructLoadView Struct : Structs)
+			{
+				// Equivalent to LoadStruct(It, Struct);
+				FBitCacheReader Bits;
+				KeyRange.LoadItem(It,					/* in-out */ Struct.Values, /* in-out */ Bits, InnerLoadIds[0], Struct.Schema.Batch);
+				ValueRange.LoadItem(It + ValueOffset,	/* in-out */ Struct.Values, /* in-out */ Bits, InnerLoadIds[1], Struct.Schema.Batch);
+				It += Stride;
+			}
 		}
 
 		Rehash(Dst);
 	}
 
 	// Load structs or ranges into non-empty map
-	inline void Insert(FScriptMap& Dst, FRangeLoadView Src) const
+	inline void InsertNonEmpty(FScriptMap& Dst, FRangeLoadView Src) const
 	{
 		// Written to avoid FProperty::CopyCompleteValue_InContainer dependency
 		// Items are loaded directly into sparse array and then removed if a duplicate existed
@@ -2462,25 +2486,34 @@ struct TMapPropertyDeltaBinding : TMapPropertyBinding<FScriptMap, KeyKind, Value
 		const int32 ValueOffset = Layout.ValueOffset;
 		int32 TmpIdx;
 		uint8* Tmp = AddPair(DstArray, /* out*/ TmpIdx, ValueOffset);
-		bool bSkip = true;
 
-		for (FStructLoadView Struct : Src.AsStructs())
+		FOptionalSchemaId InnerLoadIds[2];
+		FStructRangeLoadView Structs = Src.AsStructs();
+		Structs.GetSchema().GetInnerLoadIds(/* out */ MakeArrayView(InnerLoadIds));
+		for (FStructLoadView Struct : Structs)
 		{
 			// Equivalent to LoadStruct(It, Struct);
 			FBitCacheReader Bits;
-			KeyRange.LoadItem(Tmp,					/* in-out */ Struct.Values, /* in-out */ Bits, Struct.Schema.Batch);
-			ValueRange.LoadItem(Tmp + ValueOffset,	/* in-out */ Struct.Values, /* in-out */ Bits, Struct.Schema.Batch);
-			bSkip = HasPair(Dst, Tmp);
-			Tmp = bSkip ? Tmp : AddPair(DstArray, /* out*/ TmpIdx, ValueOffset);
-		}
+			KeyRange.LoadItem(Tmp,					/* in-out */ Struct.Values, /* in-out */ Bits, InnerLoadIds[0], Struct.Schema.Batch);
 
-		if (bSkip)
-		{
-			InnerKey.DestroyItem(Tmp);
-			InnerValue.DestroyItem(Tmp + ValueOffset);
-			DstArray.RemoveAtUninitialized(Layout.SetLayout.SparseArrayLayout, TmpIdx, 1);
+			if (int32 Idx = FindKey(Dst, Tmp); Idx != INDEX_NONE)
+			{
+				// Load value into existing pair
+				uint8* Pair = static_cast<uint8*>(DstArray.GetData(Idx, Layout.SetLayout.SparseArrayLayout));
+				ValueRange.LoadItem(Pair + ValueOffset,	/* in-out */ Struct.Values, /* in-out */ Bits, InnerLoadIds[1], Struct.Schema.Batch);
+			}
+			else
+			{
+				// Load value into tmp pair and add new temporary
+				ValueRange.LoadItem(Tmp + ValueOffset,	/* in-out */ Struct.Values, /* in-out */ Bits, InnerLoadIds[1], Struct.Schema.Batch);
+				Tmp = AddPair(DstArray, /* out*/ TmpIdx, ValueOffset);
+			}
 		}
 
+		InnerKey.DestroyItem(Tmp);
+		InnerValue.DestroyItem(Tmp + ValueOffset);
+		DstArray.RemoveAtUninitialized(Layout.SetLayout.SparseArrayLayout, TmpIdx, 1);
+
 		if (DstArray.NumUnchecked() != OldNum)
 		{
 			Rehash(Dst);
@@ -2534,18 +2567,17 @@ struct TMapPropertyDeltaBinding : TMapPropertyBinding<FScriptMap, KeyKind, Value
 
 ////////////////////////////////////////////////////////////////////////////////////////////////
 
-
 class FMapBindings
 {
-	TMap<FStructSchemaId, FRangeBinding> NormalBindings;
-	TMap<FStructSchemaId, FRangeBinding> FrozenBindings;
+	TMap<FStructSchemaId, FStructSchemaId>	NormalBindings;
+	TMap<FStructSchemaId, FRangeBinding>	FrozenBindings;
 
 	template<EPropertyKind KeyKind, EPropertyKind ValueKind>
-	static IItemRangeBinding* New3(FMapProperty* Property, FMapMemberBindings Members, bool bFreezable)
+	static IItemRangeBinding* New3(FMapProperty* Property, FMapMemberBindings Members, FStructSchemaId* OutCustomId)
 	{
-		if (bFreezable) // Freezable maps aren't delta serialized
+		if (OutCustomId == nullptr) // Freezable maps aren't delta serialized
 		{
-			return new TMapPropertyBinding<FFreezableScriptMap, KeyKind, ValueKind>(Property);
+			return new TMapPropertyItemBinding<FFreezableScriptMap, KeyKind, ValueKind>(Property);
 		}
 
 		// Index custom delta binding struct name
@@ -2553,64 +2585,74 @@ class FMapBindings
 		TConstArrayView<FTypeId> Params = GUE.Names.Resolve(PairTypename).GetParameters();
 		FTypeId Type = { GUE.Scopes.Core, FTypenameId(GUE.Names.MakeParametricTypeId(GUE.Typenames.Map, Params)) };
 		FStructSchemaId Id = GUE.Names.IndexStruct(Type);
+		FMemberId MemberIds[] = { GUE.Members.Assign, GUE.Members.Remove, GUE.Members.Insert };
 
 		// Todo: Ownership / memory leak
-		const FStructDeclaration& Declaration = GUE.Types.DeclareStruct(Id, Type, { GUE.Members.Del, GUE.Members.Add }, EMemberPresence::AllowSparse);
-		auto Out = new TMapPropertyDeltaBinding<KeyKind, ValueKind>(Property, Members);
-		GUE.DeltaCustoms.BindStruct(Declaration.Id, *Out, Declaration, {});
+		const FStructDeclaration& Declaration = GUE.Types.DeclareStruct(Id, Type, MemberIds, EMemberPresence::AllowSparse);
+		auto Out = new TMapPropertyCustomBinding<KeyKind, ValueKind>(Property, Members);
+		GUE.Customs.BindStruct(Declaration.Id, *Out, Declaration, {});
 
+		*OutCustomId = Id;
 		return Out;
 	}
 
 	template<EPropertyKind KeyKind>
-	inline IItemRangeBinding* New2(FMapProperty* Property, FMapMemberBindings Members, bool bFreezable)
+	inline IItemRangeBinding* New2(FMapProperty* Property, FMapMemberBindings Members, FStructSchemaId* OutCustomId)
 	{
 		switch (GetPropertyKind(Members.Value))
 		{
-		case EPropertyKind::Range:		return New3<KeyKind, EPropertyKind::Range>(Property, Members, bFreezable);
-		case EPropertyKind::Struct:		return New3<KeyKind, EPropertyKind::Struct>(Property, Members, bFreezable);
-		case EPropertyKind::Bool:		return New3<KeyKind, EPropertyKind::Bool>(Property, Members, bFreezable);
-		case EPropertyKind::U8:			return New3<KeyKind, EPropertyKind::U8>(Property, Members, bFreezable);
-		case EPropertyKind::U16:		return New3<KeyKind, EPropertyKind::U16>(Property, Members, bFreezable);
-		case EPropertyKind::U32:		return New3<KeyKind, EPropertyKind::U32>(Property, Members, bFreezable);
-		case EPropertyKind::U64:		return New3<KeyKind, EPropertyKind::U64>(Property, Members, bFreezable);
-		case EPropertyKind::F32:		return New3<KeyKind, EPropertyKind::F32>(Property, Members, bFreezable);
-		case EPropertyKind::F64:		return New3<KeyKind, EPropertyKind::F64>(Property, Members, bFreezable);
+		case EPropertyKind::Range:		return New3<KeyKind, EPropertyKind::Range>(Property, Members, OutCustomId);
+		case EPropertyKind::Struct:		return New3<KeyKind, EPropertyKind::Struct>(Property, Members, OutCustomId);
+		case EPropertyKind::Bool:		return New3<KeyKind, EPropertyKind::Bool>(Property, Members, OutCustomId);
+		case EPropertyKind::U8:			return New3<KeyKind, EPropertyKind::U8>(Property, Members, OutCustomId);
+		case EPropertyKind::U16:		return New3<KeyKind, EPropertyKind::U16>(Property, Members, OutCustomId);
+		case EPropertyKind::U32:		return New3<KeyKind, EPropertyKind::U32>(Property, Members, OutCustomId);
+		case EPropertyKind::U64:		return New3<KeyKind, EPropertyKind::U64>(Property, Members, OutCustomId);
+		case EPropertyKind::F32:		return New3<KeyKind, EPropertyKind::F32>(Property, Members, OutCustomId);
+		case EPropertyKind::F64:		return New3<KeyKind, EPropertyKind::F64>(Property, Members, OutCustomId);
 		default:						check(false); return nullptr;
 		}
 	}
 
-	inline IItemRangeBinding* New(FMapProperty* Property, FMapMemberBindings Members, bool bFreezable)
+	inline IItemRangeBinding* New(FMapProperty* Property, FMapMemberBindings Members, FStructSchemaId* OutCustomId)
 	{
 		switch (GetPropertyKind(Members.Key))
 		{
-		case EPropertyKind::Range:		return New2<EPropertyKind::Range>(Property, Members, bFreezable);
-		case EPropertyKind::Struct:		return New2<EPropertyKind::Struct>(Property, Members, bFreezable);
-		case EPropertyKind::Bool:		return New2<EPropertyKind::Bool>(Property, Members, bFreezable);
-		case EPropertyKind::U8:			return New2<EPropertyKind::U8>(Property, Members, bFreezable);
-		case EPropertyKind::U16:		return New2<EPropertyKind::U16>(Property, Members, bFreezable);
-		case EPropertyKind::U32:		return New2<EPropertyKind::U32>(Property, Members, bFreezable);
-		case EPropertyKind::U64:		return New2<EPropertyKind::U64>(Property, Members, bFreezable);
-		case EPropertyKind::F32:		return New2<EPropertyKind::F32>(Property, Members, bFreezable);
-		case EPropertyKind::F64:		return New2<EPropertyKind::F64>(Property, Members, bFreezable);
+		case EPropertyKind::Range:		return New2<EPropertyKind::Range>(Property, Members, OutCustomId);
+		case EPropertyKind::Struct:		return New2<EPropertyKind::Struct>(Property, Members, OutCustomId);
+		case EPropertyKind::Bool:		return New2<EPropertyKind::Bool>(Property, Members, OutCustomId);
+		case EPropertyKind::U8:			return New2<EPropertyKind::U8>(Property, Members, OutCustomId);
+		case EPropertyKind::U16:		return New2<EPropertyKind::U16>(Property, Members, OutCustomId);
+		case EPropertyKind::U32:		return New2<EPropertyKind::U32>(Property, Members, OutCustomId);
+		case EPropertyKind::U64:		return New2<EPropertyKind::U64>(Property, Members, OutCustomId);
+		case EPropertyKind::F32:		return New2<EPropertyKind::F32>(Property, Members, OutCustomId);
+		case EPropertyKind::F64:		return New2<EPropertyKind::F64>(Property, Members, OutCustomId);
 		default:						check(false); return nullptr;
 		}
 	}
 
 public:
-	FRangeBinding Bind(FMapProperty* Property, FMapMemberBindings Members)
+	FStructSchemaId BindNormal(FMapProperty* Property, FStructSchemaId PairId, FMapMemberBindings Members)
 	{
-		FStructSchemaId PairId = static_cast<FStructSchemaId>(Members.Pair.InnermostSchema.Get());
-		bool bFreezable = EnumHasAnyFlags(Property->MapFlags, EMapPropertyFlags::UsesMemoryImageAllocator);
-		TMap<FStructSchemaId, FRangeBinding>& Bindings = bFreezable ? FrozenBindings : NormalBindings;
-		if (const FRangeBinding* RangeBinding = Bindings.Find(PairId))
+		if (const FStructSchemaId* CustomId = NormalBindings.Find(PairId))
+		{
+			return *CustomId;
+		}
+
+		FStructSchemaId CustomId;
+		New(Property, Members, /* out */ &CustomId);
+		return NormalBindings.Emplace(PairId, CustomId);
+	}
+
+	FRangeBinding BindFreezable(FMapProperty* Property, FStructSchemaId PairId, FMapMemberBindings Members)
+	{
+		if (const FRangeBinding* RangeBinding = FrozenBindings.Find(PairId))
 		{
 			return *RangeBinding;
 		}
 
-		// Todo: Ownership / memory leak
-		IItemRangeBinding* Leak = New(Property, Members, bFreezable);
-		return Bindings.Emplace(PairId, FRangeBinding(*Leak, DefaultRangeMax));
+		IItemRangeBinding* Leak = New(Property, Members, nullptr);
+		return FrozenBindings.Emplace(PairId, FRangeBinding(*Leak, DefaultRangeMax));
 	}
 };
 static FMapBindings GMaps;
@@ -2908,8 +2950,11 @@ class FPropertyBinder
 public:
 	FPropertyBinder(FStructSchemaId Id) : Owner(Id) {}
 
-	void								BindMember(FProperty* Property) { Members.Emplace(BindProperty(Property)); }
-	TConstArrayView<FMemberBinding>		GetMembers() const { return Members; }
+	void								AddMember(FMemberBinding Member)	{ Members.Add(Member); }
+	void								BindMember(FProperty* Property)		{ Members.Emplace(BindProperty(Property)); }
+	TConstArrayView<FMemberBinding>		GetMembers() const					{ return Members; }
+
+	// Only true for noexport UScriptStruct with STRUCT_Immutable | STRUCT_Atomic flags
 	bool								IsDense() { return GetOwnerOccupancy() == EMemberPresence::RequireAll; }
 
 private:
@@ -3067,12 +3112,16 @@ private:
 		FMemberBinding Key = BindSingleProperty(Property->KeyProp);
 		FMemberBinding Value = BindSingleProperty(Property->ValueProp);
 		FMemberBinding Pair = GPairs.Bind(Key, Value);
-		return BindAsRange(Property, GMaps.Bind(Property, {Key, Value, Pair}), Pair);
+		FStructSchemaId PairId = static_cast<FStructSchemaId>(Pair.InnermostSchema.Get());
+
+		bool bFreezable = EnumHasAnyFlags(Property->MapFlags, EMapPropertyFlags::UsesMemoryImageAllocator);
+		return bFreezable	? BindAsRange(Property, GMaps.BindFreezable(Property, PairId, {Key, Value, Pair}), Pair)
+							: BindAsStruct(Property, GMaps.BindNormal(Property, PairId, {Key, Value, Pair}));
 	}
 	inline FMemberBinding BindSet(FSetProperty* Property)
 	{
 		FMemberBinding Elem = BindSingleProperty(Property->ElementProp);
-		return BindAsRange(Property, GSets.Bind(Property, Elem), Elem);
+		return BindAsStruct(Property, GSets.Bind(Property, Elem));
 	}
 	inline FMemberBinding BindOptional(FOptionalProperty* Property)
 	{
@@ -3132,14 +3181,14 @@ private:
 				: BindSparseDelegate(Owner, static_cast<FMulticastSparseDelegateProperty*>(Property));
 			return BindAsStruct(Property, BindId);
 		}
-	#if WITH_VERSE_VM
+#if WITH_VERSE_VM
 		else if (HasAny<CASTCLASS_FVValueProperty | CASTCLASS_FVRestValueProperty>(Flags))
 		{
 			return HasAny<CASTCLASS_FVValueProperty>(Flags)
 				? BindVValue(static_cast<FVValueProperty*>(Property))
 				: BindVRestValue(static_cast<FVRestValueProperty*>(Property)); 
 		}
-	#endif
+#else // Verse BPVM
 		else
 		{
 			if (PropertyTypename == VerseFunctionProperty) // FVerseFunctionProperty
@@ -3154,9 +3203,8 @@ private:
 			{
 				return BindAsStruct(Property, GUE.Structs.ReferencePropertyValue);
 			}
-
-			
 		}
+#endif
 
 		checkf(false, TEXT("Unrecognized class cast flags %llx in %s %s"), Flags, *PropertyTypename.ToString(), *Property->GetNameCPP());
 		return FMemberBinding(Property->GetOffset_ForInternal());
@@ -3257,9 +3305,23 @@ void BindStruct(FStructSchemaId Id, const UStruct* Struct)
 	}
 
 	FPropertyBinder Binder(Id);
-	if (Struct->GetSuperStruct() && Binder.IsDense())
+	if (Struct->GetInheritanceSuper())
 	{
-		BindSuperMembers(/* out */ Binder, Struct);
+		if (Binder.IsDense())
+		{
+			BindSuperMembers(/* out */ Binder, Struct);
+		}
+		else
+		{
+			const FStructDeclaration& Declared = GUE.Types.Get(Id);
+			if (Declared.Super)
+			{
+				FMemberBinding Member;
+				Member.InnermostType = SuperStructBindType;
+				Member.InnermostSchema = static_cast<FSchemaId>(Declared.Super.Get());
+				Binder.AddMember(Member);
+			}
+		}
 	}
 	BindMembers(/* out */ Binder, Struct);
 
@@ -3279,13 +3341,20 @@ static void BindInitialTypes()
 	UE_LOGFMT(LogPlainPropsUObject, Display, "Binding types to PlainProps schemas...");
 
 	// Declare all UScriptStruct/UClass/UFunction/UEnum
+	static const FStructSchemaId SkipStructs[] = { GUE.Structs.VerseFunction };
 	TArray<FStructSchemaId> Ids;
 	for (TObjectIterator<UField> It; It; ++It)
 	{
 		UField* Field = *It;
 		if (UStruct* Struct = Cast<UStruct>(Field))
 		{
-			Ids.Add(DeclareStruct(Struct));
+			FTypeId Type = IndexType(Struct);
+			FStructSchemaId Id = GUE.Names.IndexStruct(Type);
+			if (!Algo::Find(SkipStructs, Id))
+			{
+				DeclareStruct(Struct, Type, Id);
+			}
+			Ids.Add(Id);
 		}
 		else if (UEnum* Enum = Cast<UEnum>(Field))
 		{
@@ -3295,9 +3364,12 @@ static void BindInitialTypes()
 
 	// Bind all UScriptStruct/UClass/UFunction
 	const FStructSchemaId* IdIt = &Ids[0];
-	for (TObjectIterator<UStruct> It; It; ++It)
+	for (TObjectIterator<UStruct> It; It; ++It, ++IdIt)
 	{
-		BindStruct(*IdIt++, *It);
+		if (!Algo::Find(SkipStructs, *IdIt))
+		{
+			BindStruct(*IdIt, *It);
+		}
 	}
 	check(IdIt == &Ids[0] + Ids.Num());
 }
@@ -3501,9 +3573,9 @@ volatile static const UObject* GDebugNoteObject;
 class FBatchSaver
 {
 public:
-	FBatchSaver(FCustomBindings& Flat, FCustomBindings& Delta, int32 NumReserve)
-	: FlatCtx{GUE.Types, GUE.Schemas, Flat, Scratch, nullptr}
-	, DeltaCtx{GUE.Types, GUE.Schemas, Delta, Scratch, &GUE.Defaults}
+	FBatchSaver(FCustomBindings& Customs, int32 NumReserve)
+	: FlatCtx{GUE.Types, GUE.Schemas, Customs, Scratch, nullptr}
+	, DeltaCtx{GUE.Types, GUE.Schemas, Customs, Scratch, &GUE.Defaults}
 	{
 		SavedObjects.Reserve(NumReserve);
 	}
@@ -3517,7 +3589,7 @@ public:
 	TArray64<uint8>	Write(TArray<FStructSchemaId>& OutRuntimeIds) const
 	{
 		// Build partial schemas
-		const FStructBindIds BindIds(GUE.DeltaCustoms, GUE.Schemas);
+		const FStructBindIds BindIds(GUE.Customs, GUE.Schemas);
 		FSchemasBuilder SchemaBuilders(GUE.Types, GUE.Names, BindIds, Scratch, ESchemaFormat::InMemoryNames);
 		for (const FSavedObject& Object : SavedObjects)
 		{
@@ -3714,10 +3786,10 @@ public:
 		SetPortFlags(RoundtripPortFlags);
 	}
 
-	FORCENOINLINE void ReadProperties(UObject* Object)
+	FORCENOINLINE void ReadProperties(UObject* Object, UObject* Defaults)
 	{
 		UClass* Class = Object->GetClass();
-		Class->SerializeTaggedProperties(*this, reinterpret_cast<uint8*>(Object), Class, nullptr);
+		Class->SerializeTaggedProperties(*this, reinterpret_cast<uint8*>(Object), Class, reinterpret_cast<uint8*>(Defaults));
 	}
 
 	virtual FArchive& operator<<(FText& Value) override
@@ -3751,6 +3823,7 @@ struct FInstance
 {
 	UObject* Orig;
 	UObject* Arch;
+	UObject* Base;
 	UObject* PP;
 	UObject* TPS;
 	UObject* UPS;
@@ -3761,7 +3834,7 @@ struct FInstance
 		UClass* Class = Orig->GetClass();
 		Class->GetDefaultObject(/* create lazily */ true);
 		Arch = Orig->GetArchetype();
-		Arch = Arch == Orig ? nullptr : Arch;
+		check(Arch != Orig);
 		Id = GUE.Names.IndexStruct(IndexType(Class));
 	}
 };
@@ -3783,7 +3856,8 @@ static bool IncludeClass(UClass* Class)
 	static const FName Exclusions[] = {
 		"GameFeaturePluginStateMachine", // Cloning ensure
 		"WorldSettings", // QAGame
-		"AnimationSequencerDataModel" // CitySample - Enum value 8 is undeclared in /Script/Engine.ERichCurveTangentMode, illegal value detected in /Script/Engine.RichCurveKey::TangentMode 
+		// CitySample - Enum value 8 is undeclared in /Script/Engine.ERichCurveTangentMode, illegal value detected in /Script/Engine.RichCurveKey::TangentMode
+		"AnimSequence", "AnimationSequencerDataModel", "MovieSceneControlRigParameterSection" 
 	};
 	static const FName SuperExclusions[] = { "LevelScriptActor" };
 
@@ -3817,49 +3891,38 @@ static int32 Roundtrip()
 	}
 
 	// Create CDOs if needed and then clones for PP and TPS tests
-	UE_LOGFMT(LogPlainPropsUObject, Display, "Cloning {Num} UObjects...", Instances.Num());
+	UE_LOGFMT(LogPlainPropsUObject, Display, "Cloning {Num} UObjects 4 times...", Instances.Num());
 	for (FInstance& Instance : Instances)
 	{
 		Instance.Init();
 	}
 	FlushAsyncLoading();
 
-	FName PP_N("PP", 1);
-	for (FInstance& Instance : Instances)
-	{
-		PP_N.SetNumber(1 + &Instance - Instances.GetData());
-		Instance.PP = MakeEmptyInstance(Instance.Orig, PP_N);
-	}
-	FName TPS_N("TPS", 1);
 	for (FInstance& Instance : Instances)
 	{
-		TPS_N.SetNumber(1 + &Instance - Instances.GetData());
-		Instance.TPS = MakeEmptyInstance(Instance.Orig, TPS_N);
-	}
-	FName UPS_N("UPS", 1);
-	for (FInstance& Instance : Instances)
-	{
-		UPS_N.SetNumber(1 + &Instance - Instances.GetData());
-		Instance.UPS = MakeEmptyInstance(Instance.Orig, UPS_N);
+		uint32 N = 1 + &Instance - Instances.GetData();
+		Instance.Base = Instance.Arch ? MakeEmptyInstance(Instance.Orig, FName("Base", N)) : nullptr;
+		Instance.PP = MakeEmptyInstance(Instance.Orig, FName("PP", N));
+		Instance.TPS = MakeEmptyInstance(Instance.Orig, FName("TPS", N));
+		Instance.UPS = MakeEmptyInstance(Instance.Orig, FName("UPS", N));
 	}
 
 	// Save
 	UE_LOGFMT(LogPlainPropsUObject, Display, "Saving UObjects to PlainProps...");
 	FMemoryBatch Plain;
-	TCustomPropertyBindings<FMemoryPropertyBatch> FlatBindings(Plain.Properties, GUE.Customs);
-	TCustomPropertyBindings<FMemoryPropertyBatch> DeltaBindings(Plain.Properties, GUE.DeltaCustoms);
+	TCustomPropertyBindings<FMemoryPropertyBatch> Customs(Plain.Properties, GUE.Customs);
 	{
-		FBatchSaver Batch(FlatBindings.Overlay, DeltaBindings.Overlay, GUObjectArray.GetObjectArrayNum());
+		FBatchSaver Batch(Customs.Overlay, GUObjectArray.GetObjectArrayNum());
 		for (const FInstance& Instance : Instances)
 		{	
-			Batch.Save(Instance.Id, Instance.Orig, Instance.Arch ? Instance.PP : nullptr);
+			Batch.Save(Instance.Id, Instance.Orig, Instance.Base);
 		}
 		Plain.Data = Batch.Write(/* out */ Plain.RuntimeIds);
 	}
 
 	// Load
 	UE_LOGFMT(LogPlainPropsUObject, Display, "Loading UObjects from PlainProps...");
-	FMemoryBatchLoader Batch(DeltaBindings.Overlay, MakeMemoryView(Plain.Data), Plain.RuntimeIds);
+	FMemoryBatchLoader Batch(Customs.Overlay, MakeMemoryView(Plain.Data), Plain.RuntimeIds);
 	for (const FInstance& Instance : Instances)
 	{
 		Batch.Load(Instance.PP);
@@ -3867,19 +3930,16 @@ static int32 Roundtrip()
 
 	// Diff original vs PlainProps
 	UE_LOGFMT(LogPlainPropsUObject, Display, "Diffing UObjects roundtripped via PlainProps...");
-	FDiffContext FlatCtx	= { GUE.Types, GUE.Schemas, FlatBindings.Overlay };
-	FDiffContext DeltaCtx	= { GUE.Types, GUE.Schemas, DeltaBindings.Overlay };
+	FDiffContext DiffCtx = { GUE.Types, GUE.Schemas, Customs.Overlay };
 	TUtf8StringBuilder<256> PpDiffs;
 	uint32 NumPpDiffs = 0;
 	for (FInstance& Instance : Instances)
 	{
-		FDiffContext& Ctx = Instance.Arch ? DeltaCtx : FlatCtx;
-		if (DiffStructs(Instance.Id, Instance.Orig, Instance.PP, /* in-out */ Ctx))
+		if (DiffStructs(Instance.Id, Instance.Orig, Instance.PP, /* in-out */ DiffCtx))
 		{
-			PrintDiff(/* out */ PpDiffs, Ctx.Out, GUE.Names);
-			//Batch.Reload(Instance.TPS, &Instance - Instances.GetData());
-			//UObject* Foo = MakeEmptyInstance(Instance.Orig, "Banana");
-			Ctx.Out.Reset();
+			PrintDiff(/* out */ PpDiffs, DiffCtx.Out, GUE.Names);
+			//Batch.Reload(Instance.UPS, &Instance - Instances.GetData());
+			DiffCtx.Out.Reset();
 			PpDiffs.Append("\n");	
 			++NumPpDiffs;
 		}
@@ -3892,7 +3952,7 @@ static int32 Roundtrip()
 		FPropertyWriter Archive(/* out */ Tps);
 		for (const FInstance& Instance : Instances)
 		{
-			Archive.WriteProperties(Instance.Orig, Instance.Arch ? Instance.TPS : nullptr);
+			Archive.WriteProperties(Instance.Orig, Instance.Base);
 		}
 	}
 
@@ -3901,7 +3961,7 @@ static int32 Roundtrip()
 		FPropertyReader Archive(Tps);
 		for (const FInstance& Instance : Instances)
 		{
-			Archive.ReadProperties(Instance.TPS);
+			Archive.ReadProperties(Instance.TPS, Instance.Base);
 		}
 	}
 
@@ -3911,21 +3971,33 @@ static int32 Roundtrip()
 	TArray<int32> TpsDiffIdxs;
 	for (FInstance& Instance : Instances)
 	{
-		static const FName SkipClasses[] = { "BodySetup" }; // Skips some structs due to native FCollisionResponse::operator==
+		static const FName SkipClasses[] = {
+			"BodySetup", // Skips some structs due to native FCollisionResponse::operator==
+			"NiagaraScript", "NiagaraNodeFunctionCall", "NiagaraMeshRendererProperties", // FNiagaraTypeDefinition::Serialize resets ClassStructOrEnum
+		};
 		if (Algo::Find(SkipClasses, Instance.Orig->GetClass()->GetFName()))
 		{
 			continue;
 		}
 
-		FDiffContext& Ctx = Instance.Arch ? DeltaCtx : FlatCtx;
-		if (DiffStructs(Instance.Id, Instance.Orig, Instance.TPS, /* in-out */ Ctx))
+		if (DiffStructs(Instance.Id, Instance.Orig, Instance.TPS, /* in-out */ DiffCtx))
 		{
-			PrintDiff(/* out */ TpsDiffs, Ctx.Out, GUE.Names);
+			PrintDiff(/* out */ TpsDiffs, DiffCtx.Out, GUE.Names);
 			
-			FPropertyWriter(/* out */ Tps).WriteProperties(Instance.Orig, Instance.Arch ? Instance.UPS : nullptr);
-			Ctx.Out.Reset();
+			DiffCtx.Out.Reset();
 			TpsDiffs.Append("\n");	
 			TpsDiffIdxs.Add(&Instance - &Instances[0]);
+			
+			//FArchivedProperties Tmp;
+			//FPropertyWriter(/* out */ Tmp).WriteProperties(Instance.Orig, Instance.Base);
+			//FPropertyReader(Tmp).ReadProperties(Instance.UPS, Instance.Base);
+
+			//FDiffContext TmpCtx = DiffCtx;
+			//bool bOU = DiffStructs(Instance.Id, Instance.Orig, Instance.UPS, /* in-out */ TmpCtx);
+			//bool bOT = DiffStructs(Instance.Id, Instance.Orig, Instance.TPS, /* in-out */ TmpCtx);
+			//bool bOP = DiffStructs(Instance.Id, Instance.Orig, Instance.PP, /* in-out */ TmpCtx);
+			//TUtf8StringBuilder<256> TmpDiff;
+			//PrintDiff(/* out */ TmpDiff, TmpCtx.Out, GUE.Names);
 		}
 	}
 	UE_LOGFMT(LogPlainPropsUObject, Display, "Detected {Diffs} diffs in {Objs} UObjects saved in a {KB}KB value stream using TPS", TpsDiffIdxs.Num(), Instances.Num(), Tps.Data.NumBytes() / 1024);
@@ -3970,7 +4042,7 @@ int32 UTestPlainPropsCommandlet::Main(const FString& Params)
 {
 	if (int32 LoadIdx = Params.Find("-load="); LoadIdx != INDEX_NONE)
 	{
-		// E.g. -run=TestPlainProps -load=/BRRoot/BRRoot.BRRoot,/Game/Maps/FrontEnd.Front
+		// E.g. -run=TestPlainProps -load=/BRRoot/BRRoot.BRRoot,/Game/Maps/FrontEnd.FrontEnd
 
 		const TCHAR* It = &Params[LoadIdx + 6];
 		FStringView Assets(It, FAsciiSet::FindFirstOrEnd(It, FAsciiSet(" ")) - It);
diff --git a/Engine/Plugins/Experimental/PlainPropsUObject/Source/Public/PlainPropsUObjectRuntime.h b/Engine/Plugins/Experimental/PlainPropsUObject/Source/Public/PlainPropsUObjectRuntime.h
--- a/Engine/Plugins/Experimental/PlainPropsUObject/Source/Public/PlainPropsUObjectRuntime.h
+++ b/Engine/Plugins/Experimental/PlainPropsUObject/Source/Public/PlainPropsUObjectRuntime.h
@@ -120,8 +120,9 @@ struct FCommonMemberIds
 	
 	FMemberId					Key;
 	FMemberId					Value;
-	FMemberId					Add;
-	FMemberId					Del;
+	FMemberId					Assign;
+	FMemberId					Remove;
+	FMemberId					Insert;
 	FMemberId					Id;
 	FMemberId					Object;
 	FMemberId					Function;
@@ -138,7 +139,6 @@ struct FGlobals
 	FDeclarations				Types;
 	FSchemaBindings				Schemas;
 	FCustomBindingsBottom		Customs;
-	FCustomBindingsOverlay		DeltaCustoms;
 	FDefaultStructs				Defaults;
 	FCommonScopeIds				Scopes;
 	FCommonStructIds			Structs;
@@ -175,7 +175,6 @@ struct FDeltaRuntime : FDefaultRuntime
 {
 	template<class T> using CustomBindings = TCustomDeltaBind<T>;
 
-	static FCustomBindings&		GetCustoms()		{ return GUE.DeltaCustoms; }
 	static IDefaultStructs*		GetDefaults()		{ return &GUE.Defaults; }
 };
 
