PlainProps fix delegate and FClassProperty binding, improve diff output and some minor load opts

#rnx

diff --git a/Engine/Plugins/Experimental/PlainProps/Source/Private/PlainPropsLoad.cpp b/Engine/Plugins/Experimental/PlainProps/Source/Private/PlainPropsLoad.cpp
--- a/Engine/Plugins/Experimental/PlainProps/Source/Private/PlainPropsLoad.cpp
+++ b/Engine/Plugins/Experimental/PlainProps/Source/Private/PlainPropsLoad.cpp
@@ -1027,12 +1027,14 @@ private:
 	int32 InnerRangeIdx = 0;
 	int32 InnerStructIdx = 0;	
 
-	void SkipMissingSparseMembers()
+	void SkipMissingSparseMembers() {}
+	void SkipMissingSparseMembers() requires (bSparse)
 	{
 		// Make code changes in FMemberReader::SkipMissingSparseMembers() too
-		while (bSparse && MemberIdx < Types.Num() && BitIt.GrabNext(ByteIt))
+
+		while (MemberIdx < Types.Num() && BitIt.GrabNext(ByteIt))
 		{
-			FMemberBindType Type = Types[MemberIdx];		
+			FMemberBindType Type = Types.GetData()[MemberIdx];		
 			EMemberKind Kind = Type.GetKind();
 			if (Kind == EMemberKind::Struct)
 			{
@@ -1060,7 +1062,7 @@ private:
 		}
 	}
 
-	FStructSchemaId GrabInnerStruct(FStructBindType Type)
+	inline FStructSchemaId GrabInnerStruct(FStructBindType Type)
 	{
 		return Type.IsDynamic ? FStructSchemaId { ByteIt.Grab<uint32>() } : InnerStructSchemas[InnerStructIdx++];
 	}
diff --git a/Engine/Plugins/Experimental/PlainProps/Source/Public/PlainPropsBind.h b/Engine/Plugins/Experimental/PlainProps/Source/Public/PlainPropsBind.h
--- a/Engine/Plugins/Experimental/PlainProps/Source/Public/PlainPropsBind.h
+++ b/Engine/Plugins/Experimental/PlainProps/Source/Public/PlainPropsBind.h
@@ -112,9 +112,9 @@ union FMemberBindType
 	bool					IsStruct() const	{ return Kind == EMemberKind::Struct; }
 	EMemberKind				GetKind() const		{ return Kind; }
 	
-	FLeafBindType			AsLeaf() const		{ check(IsLeaf());		return Leaf; }
-	FRangeBindType			AsRange() const		{ check(IsRange());		return Range; }
-	FStructBindType			AsStruct() const	{ check(IsStruct());	return Struct; }
+	FLeafBindType			AsLeaf() const		{ checkSlow(IsLeaf());		return Leaf; }
+	FRangeBindType			AsRange() const		{ checkSlow(IsRange());		return Range; }
+	FStructBindType			AsStruct() const	{ checkSlow(IsStruct());	return Struct; }
 	uint8					AsByte() const		{ return BitCast<uint8>(*this); }
 
 	friend inline bool operator==(FMemberBindType A, FMemberBindType B) { return A.AsByte() == B.AsByte(); }
diff --git a/Engine/Plugins/Experimental/PlainProps/Source/Public/PlainPropsRead.h b/Engine/Plugins/Experimental/PlainProps/Source/Public/PlainPropsRead.h
--- a/Engine/Plugins/Experimental/PlainProps/Source/Public/PlainPropsRead.h
+++ b/Engine/Plugins/Experimental/PlainProps/Source/Public/PlainPropsRead.h
@@ -58,7 +58,7 @@ public:
 #endif
 	{}
 
-	[[nodiscard]] const uint8* GrabBytes(uint64 NumBytes)
+	[[nodiscard]] inline const uint8* GrabBytes(uint64 NumBytes)
 	{
 #if DO_CHECK
 		check(It + NumBytes <= End);
@@ -89,10 +89,10 @@ public:
 		return FPlatformMemory::ReadUnaligned<T>(GrabBytes(sizeof(T)));
 	}
 	
-	[[nodiscard]] uint64 GrabVarIntU()
+	[[nodiscard]] FORCEINLINE uint64 GrabVarIntU()
 	{
 #if DO_CHECK
-		check(It < End);
+		checkSlow(It < End);
 #endif
 		uint32 NumBytesRead;
 		uint64 Out = ReadVarUInt(It, NumBytesRead);
diff --git a/Engine/Plugins/Experimental/PlainPropsUObject/Source/Private/PlainPropsUObjectBindings.cpp b/Engine/Plugins/Experimental/PlainPropsUObject/Source/Private/PlainPropsUObjectBindings.cpp
--- a/Engine/Plugins/Experimental/PlainPropsUObject/Source/Private/PlainPropsUObjectBindings.cpp
+++ b/Engine/Plugins/Experimental/PlainPropsUObject/Source/Private/PlainPropsUObjectBindings.cpp
@@ -296,6 +296,7 @@ FCommonStructIds::FCommonStructIds(const FCommonScopeIds& Scopes, TIdIndexer<FNa
 , Guid(						Names.IndexStruct({Scopes.Core, Names.MakeTypename("Guid")}))
 , FieldPath(				Names.IndexStruct({Scopes.CoreUObject, Names.MakeTypename("FieldPath")}))
 , SoftObjectPath(			Names.IndexStruct({Scopes.CoreUObject, Names.MakeTypename("SoftObjectPath")}))
+, ClassPtr(					Names.IndexStruct({Scopes.CoreUObject, Names.MakeTypename("ClassPtr")}))
 , ObjectPtr(				Names.IndexStruct({Scopes.CoreUObject, Names.MakeTypename("ObjectPtr")}))
 , WeakObjectPtr(			Names.IndexStruct({Scopes.CoreUObject, Names.MakeTypename("WeakObjectPtr")}))
 , LazyObjectPtr(			Names.IndexStruct({Scopes.CoreUObject, Names.MakeTypename("LazyObjectPtr")}))
@@ -346,7 +347,7 @@ static constexpr uint64 IntSMask =			CASTCLASS_FInt8Property | CASTCLASS_FInt16P
 static constexpr uint64 IntUMask =			CASTCLASS_FByteProperty | CASTCLASS_FUInt16Property | CASTCLASS_FUInt32Property | CASTCLASS_FUInt64Property;
 static constexpr uint64 ContainerMask =		CASTCLASS_FArrayProperty | CASTCLASS_FSetProperty | CASTCLASS_FMapProperty  | CASTCLASS_FOptionalProperty;
 static constexpr uint64 StringMask =		CASTCLASS_FStrProperty | CASTCLASS_FUtf8StrProperty | CASTCLASS_FAnsiStrProperty | CASTCLASS_FVerseStringProperty;
-static constexpr uint64 CommonStructMask =	CASTCLASS_FNameProperty | CASTCLASS_FTextProperty |  CASTCLASS_FFieldPathProperty |
+static constexpr uint64 CommonStructMask =	CASTCLASS_FNameProperty | CASTCLASS_FTextProperty |  CASTCLASS_FFieldPathProperty | CASTCLASS_FClassProperty |
 											CASTCLASS_FObjectProperty | CASTCLASS_FWeakObjectProperty | CASTCLASS_FSoftObjectProperty | CASTCLASS_FLazyObjectProperty |
 											CASTCLASS_FDelegateProperty | CASTCLASS_FMulticastInlineDelegateProperty;
 static constexpr uint64 MiscMask =			CASTCLASS_FMulticastSparseDelegateProperty | CASTCLASS_FInterfaceProperty;
@@ -355,22 +356,23 @@ static FStructSchemaId FlagsToCommonBindId(uint64 MaskedCastFlags)
 {
 	switch (MaskedCastFlags)
 	{
-		case CASTCLASS_FNameProperty:						return GUE.Structs.Name;
-		case CASTCLASS_FObjectProperty:						return GUE.Structs.ObjectPtr;
-		case CASTCLASS_FWeakObjectProperty:					return GUE.Structs.WeakObjectPtr;
-		case CASTCLASS_FSoftObjectProperty:					return GUE.Structs.SoftObjectPtr;
-		case CASTCLASS_FLazyObjectProperty:					return GUE.Structs.LazyObjectPtr;
-		case CASTCLASS_FDelegateProperty:					return GUE.Structs.Delegate;
-		case CASTCLASS_FMulticastInlineDelegateProperty:	return GUE.Structs.MulticastInlineDelegate;
-		case CASTCLASS_FTextProperty:						return GUE.Structs.Text;
-		case CASTCLASS_FFieldPathProperty:					return GUE.Structs.FieldPath;
-		default:											break; // error
+		case CASTCLASS_FNameProperty:								return GUE.Structs.Name;
+		case CASTCLASS_FClassProperty | CASTCLASS_FObjectProperty:	return GUE.Structs.ClassPtr;
+		case CASTCLASS_FObjectProperty:								return GUE.Structs.ObjectPtr;
+		case CASTCLASS_FWeakObjectProperty:							return GUE.Structs.WeakObjectPtr;
+		case CASTCLASS_FSoftObjectProperty:							return GUE.Structs.SoftObjectPtr;
+		case CASTCLASS_FLazyObjectProperty:							return GUE.Structs.LazyObjectPtr;
+		case CASTCLASS_FDelegateProperty:							return GUE.Structs.Delegate;
+		case CASTCLASS_FMulticastInlineDelegateProperty:			return GUE.Structs.MulticastInlineDelegate;
+		case CASTCLASS_FTextProperty:								return GUE.Structs.Text;
+		case CASTCLASS_FFieldPathProperty:							return GUE.Structs.FieldPath;
+		default:													break; // error
 	}
 
 	check(MaskedCastFlags); // @pre violated
 	check((MaskedCastFlags & CommonStructMask) == MaskedCastFlags); // @pre violated
 	checkf(FMath::CountBits(MaskedCastFlags) == 1, TEXT("Masked CASTCLASS flags %llx match more than one common property type"), MaskedCastFlags);
-	check(false); // Mismatch between this function and BindAsStructMask
+	check(false); // Mismatch between this function and CommonStructMask
 	return {};
 }
 
@@ -568,11 +570,11 @@ bool FFieldPathBinding::Diff(const FFieldPath& A, const FFieldPath& B)
 
 void FDelegateBinding::Save(FMemberBuilder& Dst, const FScriptDelegate& Src, const FScriptDelegate* Default, const FSaveContext& Ctx) const
 {
-	if (Src.IsBound())
+	FDelegateAccess::FReadScope Scope(Src);
+	if (FName Function = Src.*GDelegateFunctionName; Function != FName())
 	{
-		FDelegateAccess::FReadScope Scope(Src);
 		Dst.AddStruct(GUE.Members.Object, GUE.Structs.WeakObjectPtr, SaveStruct(&(Src.*GDelegateObject), GUE.Structs.WeakObjectPtr, Ctx));
-		Dst.AddStruct(GUE.Members.Function, GUE.Structs.Name, SaveStruct(&(Src.*GDelegateFunctionName), GUE.Structs.Name, Ctx));
+		Dst.AddStruct(GUE.Members.Function, GUE.Structs.Name, SaveStruct(&Function, GUE.Structs.Name, Ctx));
 	}
 }
 
@@ -3384,6 +3386,7 @@ static void InitBatchedProperties()
 	// All property batch types share declaration, but some members could be exclusive to a specific type e.g. FMemoryPropertyBatch
 	DeclareProperty(GUE.Structs.Name,			{ GUE.Members.Id }, EMemberPresence::RequireAll);
 	DeclareProperty(GUE.Structs.Text,			{ GUE.Members.Id }, EMemberPresence::RequireAll);
+	DeclareProperty(GUE.Structs.ClassPtr,		{ GUE.Members.Id }, EMemberPresence::RequireAll);
 	DeclareProperty(GUE.Structs.ObjectPtr,		{ GUE.Members.Id }, EMemberPresence::RequireAll);
 	DeclareProperty(GUE.Structs.WeakObjectPtr,	{ GUE.Members.Id }, EMemberPresence::RequireAll);
 	DeclareProperty(GUE.Structs.SoftObjectPtr,	{ GUE.Members.Id }, EMemberPresence::RequireAll);
@@ -3391,6 +3394,7 @@ static void InitBatchedProperties()
 
 	GUE.Defaults.BindZeroes(GUE.Structs.Name, sizeof(FName), alignof(FName));
 	GUE.Defaults.BindStatic(GUE.Structs.Text, &FText::GetEmpty());
+	GUE.Defaults.BindZeroes(GUE.Structs.ClassPtr, sizeof(TSubclassOf<UClass>), alignof(TSubclassOf<UClass>));
 	GUE.Defaults.BindZeroes(GUE.Structs.ObjectPtr, sizeof(FObjectPtr), alignof(FObjectPtr));
 	GUE.Defaults.BindZeroes(GUE.Structs.WeakObjectPtr, sizeof(FWeakObjectPtr), alignof(FWeakObjectPtr));
 	GUE.Defaults.BindZeroes(GUE.Structs.SoftObjectPtr, sizeof(FSoftObjectPtr), alignof(FSoftObjectPtr));
@@ -3538,6 +3542,7 @@ struct TCustomPropertyBindings
 	{
 		Bind(GUE.Structs.Name, Name);
 		Bind(GUE.Structs.Text, Text);
+		Bind(GUE.Structs.ClassPtr, ObjectPtr); // TSubclassOf<> is essentially a TObjectPtr
 		Bind(GUE.Structs.ObjectPtr, ObjectPtr);
 		Bind(GUE.Structs.SoftObjectPtr, SoftObjectPtr);
 		Bind(GUE.Structs.WeakObjectPtr, WeakObjectPtr);
@@ -3866,6 +3871,13 @@ static bool IncludeClass(UClass* Class)
 		return false;
 	}
 
+	// Exclude IDOs
+	static constexpr EClassFlags IdoFlags = CLASS_NotPlaceable | CLASS_Hidden | CLASS_HideDropDown;
+	if (Class->HasAllClassFlags(IdoFlags))
+	{
+		return false;
+	}
+
 	for (UStruct* Super = Class->GetInheritanceSuper(); Super; Super = Super->GetInheritanceSuper())
 	{
 		if (!!Algo::Find(SuperExclusions, Super->GetFName()))
@@ -3966,7 +3978,9 @@ static int32 Roundtrip()
 			PrintDiff(/* out */ PpDiffs, DiffCtx.Out, GUE.Names);
 			//Batch.Reload(Instance.UPS, &Instance - Instances.GetData());
 			DiffCtx.Out.Reset();
-			PpDiffs.Append("\n");	
+			PpDiffs.Append(" in ");
+			PpDiffs.Append(Instance.Orig->GetFullName());
+			PpDiffs.Append("\n");
 			++NumPpDiffs;
 		}
 	}
@@ -4005,6 +4019,8 @@ static int32 Roundtrip()
 			PrintDiff(/* out */ TpsDiffs, DiffCtx.Out, GUE.Names);
 			
 			DiffCtx.Out.Reset();
+			TpsDiffs.Append(" in ");
+			TpsDiffs.Append(Instance.Orig->GetFullName());
 			TpsDiffs.Append("\n");	
 			TpsDiffIdxs.Add(&Instance - &Instances[0]);
 			
diff --git a/Engine/Plugins/Experimental/PlainPropsUObject/Source/Public/PlainPropsUObjectRuntime.h b/Engine/Plugins/Experimental/PlainPropsUObject/Source/Public/PlainPropsUObjectRuntime.h
--- a/Engine/Plugins/Experimental/PlainPropsUObject/Source/Public/PlainPropsUObjectRuntime.h
+++ b/Engine/Plugins/Experimental/PlainPropsUObject/Source/Public/PlainPropsUObjectRuntime.h
@@ -100,6 +100,7 @@ struct FCommonStructIds
 	FStructSchemaId				Guid;
 	FStructSchemaId				FieldPath;
 	FStructSchemaId				SoftObjectPath;
+	FStructSchemaId				ClassPtr;
 	FStructSchemaId				ObjectPtr;
 	FStructSchemaId				WeakObjectPtr;
 	FStructSchemaId				LazyObjectPtr;
