PlainProps test commandlet add UPS comparison

#rnx

diff --git a/Engine/Plugins/Experimental/PlainPropsUObject/Source/Private/PlainPropsUObjectBindings.cpp b/Engine/Plugins/Experimental/PlainPropsUObject/Source/Private/PlainPropsUObjectBindings.cpp
--- a/Engine/Plugins/Experimental/PlainPropsUObject/Source/Private/PlainPropsUObjectBindings.cpp
+++ b/Engine/Plugins/Experimental/PlainPropsUObject/Source/Private/PlainPropsUObjectBindings.cpp
@@ -3904,22 +3904,33 @@ FORCENOINLINE static void LoadPlainProps(FMemoryBatchLoader& Batch, TConstArrayV
 	}
 }
 
-FORCENOINLINE static void SaveTps(FPropertyWriter& Archive, TConstArrayView<FInstance> Instances)
+template<bool bUps>
+FORCENOINLINE void SaveArchive(FPropertyWriter& Archive, TConstArrayView<FInstance> Instances)
 {
+	Archive.SetUseUnversionedPropertySerialization(bUps);
 	for (const FInstance& Instance : Instances)
 	{
 		Archive.WriteProperties(Instance.Orig, Instance.Base);
 	}
 }
 
-FORCENOINLINE static void LoadTps(FPropertyReader& Archive, TConstArrayView<FInstance> Instances)
+template<bool bUps>
+FORCENOINLINE void LoadArchive(FPropertyReader& Archive, TConstArrayView<FInstance> Instances)
 {
+	Archive.SetUseUnversionedPropertySerialization(bUps);
 	for (const FInstance& Instance : Instances)
 	{
-		Archive.ReadProperties(Instance.TPS, Instance.Base);
+		Archive.ReadProperties(bUps ? Instance.UPS : Instance.TPS, Instance.Base);
 	}
 }
 
+struct FDiffDebug
+{
+	volatile FInstance* Last;
+	volatile const UTF8CHAR* Str;
+};
+static FDiffDebug GPpDiff, GTpsDiff, GUpsDiff;
+
 static int32 Roundtrip()
 {
 	UE_LOGFMT(LogPlainPropsUObject, Display, "Gathering all non-empty UObjects...");
@@ -3975,8 +3986,9 @@ static int32 Roundtrip()
 	{
 		if (DiffStructs(Instance.Id, Instance.Orig, Instance.PP, /* in-out */ DiffCtx))
 		{
+			GPpDiff.Last = &Instance;
 			PrintDiff(/* out */ PpDiffs, DiffCtx.Out, GUE.Names);
-			//Batch.Reload(Instance.UPS, &Instance - Instances.GetData());
+			Batch.Reload(Instance.UPS, &Instance - Instances.GetData());
 			DiffCtx.Out.Reset();
 			PpDiffs.Append(" in ");
 			PpDiffs.Append(Instance.Orig->GetFullName());
@@ -3984,31 +3996,33 @@ static int32 Roundtrip()
 			++NumPpDiffs;
 		}
 	}
+	GPpDiff.Str = PpDiffs.GetData();
 	UE_LOGFMT(LogPlainPropsUObject, Display, "Detected {Diffs} diffs in {Objs} UObjects saved in a {KB}KB value stream using PlainProps", NumPpDiffs, Instances.Num(), Plain.Data.NumBytes() / 1024);
 
 	UE_LOGFMT(LogPlainPropsUObject, Display, "Saving UObjects to TPS archive...");
 	FArchivedProperties Tps;
 	{
 		FPropertyWriter Archive(/* out */ Tps);
-		SaveTps(Archive, Instances);
+		SaveArchive<false>(Archive, Instances);
 	}
 
 	UE_LOGFMT(LogPlainPropsUObject, Display, "Loading UObjects from TPS archive...");
 	{
 		FPropertyReader Archive(Tps);
-		LoadTps(Archive, Instances);
+		LoadArchive<false>(Archive, Instances);
 	}
 
 	// Diff original vs TPS
 	UE_LOGFMT(LogPlainPropsUObject, Display, "Diffing UObjects roundtripped via TPS...");
+	static const FName SkipClasses[] = {
+		"BodySetup", // Skips some structs due to native FCollisionResponse::operator==
+		"NiagaraScript", "NiagaraNodeFunctionCall", "NiagaraMeshRendererProperties", // FNiagaraTypeDefinition::Serialize resets ClassStructOrEnum
+	};
+
 	TUtf8StringBuilder<256> TpsDiffs;
 	TArray<int32> TpsDiffIdxs;
 	for (FInstance& Instance : Instances)
 	{
-		static const FName SkipClasses[] = {
-			"BodySetup", // Skips some structs due to native FCollisionResponse::operator==
-			"NiagaraScript", "NiagaraNodeFunctionCall", "NiagaraMeshRendererProperties", // FNiagaraTypeDefinition::Serialize resets ClassStructOrEnum
-		};
 		if (Algo::Find(SkipClasses, Instance.Orig->GetClass()->GetFName()))
 		{
 			continue;
@@ -4016,6 +4030,7 @@ static int32 Roundtrip()
 
 		if (DiffStructs(Instance.Id, Instance.Orig, Instance.TPS, /* in-out */ DiffCtx))
 		{
+			GTpsDiff.Last = &Instance;
 			PrintDiff(/* out */ TpsDiffs, DiffCtx.Out, GUE.Names);
 			
 			DiffCtx.Out.Reset();
@@ -4036,8 +4051,48 @@ static int32 Roundtrip()
 			//PrintDiff(/* out */ TmpDiff, TmpCtx.Out, GUE.Names);
 		}
 	}
+	GTpsDiff.Str = TpsDiffs.GetData();
 	UE_LOGFMT(LogPlainPropsUObject, Display, "Detected {Diffs} diffs in {Objs} UObjects saved in a {KB}KB value stream using TPS", TpsDiffIdxs.Num(), Instances.Num(), Tps.Data.NumBytes() / 1024);
 
+	UE_LOGFMT(LogPlainPropsUObject, Display, "Saving UObjects to UPS archive...");
+	FArchivedProperties Ups;
+	{
+		FPropertyWriter Archive(/* out */ Ups);
+		SaveArchive<true>(Archive, Instances);
+	}
+
+	UE_LOGFMT(LogPlainPropsUObject, Display, "Loading UObjects from UPS archive...");
+	{
+		FPropertyReader Archive(Ups);
+		LoadArchive<true>(Archive, Instances);
+	}
+
+	// Diff original vs UPS
+	UE_LOGFMT(LogPlainPropsUObject, Display, "Diffing UObjects roundtripped via UPS...");
+	TUtf8StringBuilder<256> UpsDiffs;
+	TArray<int32> UpsDiffIdxs;
+	for (FInstance& Instance : Instances)
+	{
+		if (Algo::Find(SkipClasses, Instance.Orig->GetClass()->GetFName()))
+		{
+			continue;
+		}
+
+		if (DiffStructs(Instance.Id, Instance.Orig, Instance.UPS, /* in-out */ DiffCtx))
+		{
+			GUpsDiff.Last = &Instance;
+			PrintDiff(/* out */ UpsDiffs, DiffCtx.Out, GUE.Names);
+			
+			DiffCtx.Out.Reset();
+			UpsDiffs.Append(" in ");
+			UpsDiffs.Append(Instance.Orig->GetFullName());
+			UpsDiffs.Append("\n");	
+			UpsDiffIdxs.Add(&Instance - &Instances[0]);
+		}
+	}
+	GUpsDiff.Str = UpsDiffs.GetData();
+	UE_LOGFMT(LogPlainPropsUObject, Display, "Detected {Diffs} diffs in {Objs} UObjects saved in a {KB}KB value stream using UPS", UpsDiffIdxs.Num(), Instances.Num(), Ups.Data.NumBytes() / 1024);
+
 	return NumPpDiffs;
 }
 
