PlainProps rename FTypeId to FType

#rnx

diff --git a/Engine/Plugins/Experimental/PlainProps/Source/Private/PlainPropsBind.cpp b/Engine/Plugins/Experimental/PlainProps/Source/Private/PlainPropsBind.cpp
--- a/Engine/Plugins/Experimental/PlainProps/Source/Private/PlainPropsBind.cpp
+++ b/Engine/Plugins/Experimental/PlainProps/Source/Private/PlainPropsBind.cpp
@@ -334,7 +334,7 @@ FCustomBindingHandle FCustomBindingsOverlay::Find(FStructSchemaId BindId) const
 
 ////////////////////////////////////////////////////////////////////////////////////////////////
 
-inline FTypeId IndexInnermost(FIdIndexerBase& Ids, FMemberBindType Type, FOptionalSchemaId Schema)
+inline FType IndexInnermost(FIdIndexerBase& Ids, FMemberBindType Type, FOptionalSchemaId Schema)
 {
 	if (Type.IsStruct())
 	{
@@ -350,14 +350,14 @@ inline FTypeId IndexInnermost(FIdIndexerBase& Ids, FMemberBindType Type, FOption
 	return Ids.MakeLeafParameter(Leaf);
 }
 
-FDualTypeId	FMemberBinding::IndexParameterName(FIdIndexerBase& Ids) const
+FDualType	FMemberBinding::IndexParameterName(FIdIndexerBase& Ids) const
 {
-	FTypeId BindType = IndexInnermost(Ids, InnermostType, InnermostSchema);
-	FTypeId DeclType = BindType;
+	FType BindType = IndexInnermost(Ids, InnermostType, InnermostSchema);
+	FType DeclType = BindType;
 	for (FRangeBinding Range : ReverseIterate(RangeBindings))
 	{
-		FTypeId SizeType = Ids.MakeRangeParameter(Range.GetSizeType());
-		FTypeId RangeBindName = {NoId, FTypenameId(Range.GetBindName())};
+		FType SizeType = Ids.MakeRangeParameter(Range.GetSizeType());
+		FType RangeBindName = {NoId, FTypenameId(Range.GetBindName())};
 		DeclType = Ids.MakeAnonymousParametricType({DeclType, SizeType});
 		BindType = Ids.MakeParametricType(RangeBindName, {BindType, SizeType});
 	}
@@ -496,14 +496,14 @@ static void TranslateScopeIds(TArrayView<FNestedScopeId> Out, FIdIndexerBase& In
 	}
 }
 
-static void TranslateParametricTypeIds(TArrayView<FParametricTypeId> Out, FIdIndexerBase& Indexer, const FIdBinding& To, TConstArrayView<FParametricType> From, const FTypeId* FromParameters)
+static void TranslateParametricTypeIds(TArrayView<FParametricTypeId> Out, FIdIndexerBase& Indexer, const FIdBinding& To, TConstArrayView<FParametricType> From, const FType* FromParameters)
 {
-	TArray<FTypeId, TInlineAllocator<8>> Params;
+	TArray<FType, TInlineAllocator<8>> Params;
 	uint32 OutIdx = 0;
 	for (FParametricType Parametric : From)
 	{
 		Params.Reset();
-		for (FTypeId FromParameter : MakeArrayView(FromParameters + Parametric.Parameters.Idx, Parametric.Parameters.NumParameters))
+		for (FType FromParameter : MakeArrayView(FromParameters + Parametric.Parameters.Idx, Parametric.Parameters.NumParameters))
 		{
 			Params.Add(To.Remap(FromParameter));
 		}
@@ -516,14 +516,14 @@ static void TranslateSchemaIds(TArrayView<FSchemaId> Out, FIdIndexerBase& Indexe
 	uint32 OutIdx = 0;
 	for (const FStructSchema& FromSchema : GetStructSchemas(From))
 	{
-		FTypeId ToType = To.Remap(FromSchema.Type);
+		FType ToType = To.Remap(FromSchema.Type);
 		checkSlow(ToType.Name.NumParameters == FromSchema.Type.Name.NumParameters);
 		Out[OutIdx++] = Indexer.IndexStruct(ToType);
 	}
 	
 	for (const FEnumSchema& FromSchema : GetEnumSchemas(From))
 	{
-		FTypeId ToType = To.Remap(FromSchema.Type);
+		FType ToType = To.Remap(FromSchema.Type);
 		Out[OutIdx++] = Indexer.IndexEnum(ToType);
 	}
 }
diff --git a/Engine/Plugins/Experimental/PlainProps/Source/Private/PlainPropsBuildSchema.cpp b/Engine/Plugins/Experimental/PlainProps/Source/Private/PlainPropsBuildSchema.cpp
--- a/Engine/Plugins/Experimental/PlainProps/Source/Private/PlainPropsBuildSchema.cpp
+++ b/Engine/Plugins/Experimental/PlainProps/Source/Private/PlainPropsBuildSchema.cpp
@@ -356,7 +356,7 @@ void FStructSchemaBuilder::NoteRangeRecursively(ERangeSizeType NumType, TConstAr
 
 FBuiltStructSchema FStructSchemaBuilder::Build() const
 {
-	FTypeId Type = Id == Declaration.Id ? Declaration.Type : AllSchemas.GetIds().Resolve(Id);
+	FType Type = Id == Declaration.Id ? Declaration.Type : AllSchemas.GetIds().Resolve(Id);
 	FBuiltStructSchema Out = { Type, Id, Declaration.Super, /* dense */ true };
 	
 	if (int32 Num = NotedMembers.CountSetBits())
diff --git a/Engine/Plugins/Experimental/PlainProps/Source/Private/PlainPropsDeclare.cpp b/Engine/Plugins/Experimental/PlainProps/Source/Private/PlainPropsDeclare.cpp
--- a/Engine/Plugins/Experimental/PlainProps/Source/Private/PlainPropsDeclare.cpp
+++ b/Engine/Plugins/Experimental/PlainProps/Source/Private/PlainPropsDeclare.cpp
@@ -97,7 +97,7 @@ void CopyItems(T* It, TConstArrayView<T> Items)
 	}
 }
 
-static const FStructDeclaration& Declare(TArray<TUniquePtr<FStructDeclaration>>& Out, FStructSchemaId DeclId, FTypeId Type, TConstArrayView<FMemberId> MemberOrder, EMemberPresence Occupancy, FOptionalStructSchemaId Super, uint16 NumMembers)
+static const FStructDeclaration& Declare(TArray<TUniquePtr<FStructDeclaration>>& Out, FStructSchemaId DeclId, FType Type, TConstArrayView<FMemberId> MemberOrder, EMemberPresence Occupancy, FOptionalStructSchemaId Super, uint16 NumMembers)
 {
 	if (static_cast<int32>(DeclId.Idx) >= Out.Num())
 	{
@@ -127,18 +127,18 @@ static const FStructDeclaration& Declare(TArray<TUniquePtr<FStructDeclaration>>&
 	return *Ptr;
 }
 
-const FStructDeclaration& FDeclarations::DeclareStruct(FStructSchemaId DeclId, FTypeId Type, TConstArrayView<FMemberId> MemberOrder, EMemberPresence Occupancy, FOptionalStructSchemaId Super)
+const FStructDeclaration& FDeclarations::DeclareStruct(FStructSchemaId DeclId, FType Type, TConstArrayView<FMemberId> MemberOrder, EMemberPresence Occupancy, FOptionalStructSchemaId Super)
 {
 	checkf(!(Super && Occupancy == EMemberPresence::RequireAll), TEXT("'%s' is a dense substruct, this isn't supported, see BuildSuperStruct()"), *Debug.Print(DeclId));
 	return Declare(/* out */ DeclaredStructs, DeclId, Type, MemberOrder, Occupancy, Super, IntCastChecked<uint16>(MemberOrder.Num()));
 }
 
-const FStructDeclaration& FDeclarations::DeclareNumeralStruct(FStructSchemaId DeclId, FTypeId Type, TConstArrayView<FMemberId> Numerals, EMemberPresence Occupancy)
+const FStructDeclaration& FDeclarations::DeclareNumeralStruct(FStructSchemaId DeclId, FType Type, TConstArrayView<FMemberId> Numerals, EMemberPresence Occupancy)
 {
 	return Declare(/* out */ DeclaredStructs, DeclId, Type, Numerals, Occupancy, NoId, IntCastChecked<uint16>(Numerals.Num()));
 }
 
-const FEnumDeclaration& FDeclarations::DeclareEnum(FEnumSchemaId Id, FTypeId Type, EEnumMode Mode, TConstArrayView<FEnumerator> Enumerators, EEnumAliases Policy)
+const FEnumDeclaration& FDeclarations::DeclareEnum(FEnumSchemaId Id, FType Type, EEnumMode Mode, TConstArrayView<FEnumerator> Enumerators, EEnumAliases Policy)
 {
 	if (static_cast<int32>(Id.Idx) >= DeclaredEnums.Num())
 	{
diff --git a/Engine/Plugins/Experimental/PlainProps/Source/Private/PlainPropsIndex.cpp b/Engine/Plugins/Experimental/PlainProps/Source/Private/PlainPropsIndex.cpp
--- a/Engine/Plugins/Experimental/PlainProps/Source/Private/PlainPropsIndex.cpp
+++ b/Engine/Plugins/Experimental/PlainProps/Source/Private/PlainPropsIndex.cpp
@@ -24,23 +24,23 @@ FNestedScope FNestedScopeIndexer::Resolve(FNestedScopeId Id) const
 
 //////////////////////////////////////////////////////////////////////////
 
-FParametricTypeView::FParametricTypeView(FConcreteTypenameId InName, TConstArrayView<FTypeId> Params)
+FParametricTypeView::FParametricTypeView(FConcreteTypenameId InName, TConstArrayView<FType> Params)
 : FParametricTypeView(InName, IntCastChecked<uint8>(Params.Num()), Params.GetData())
 {}
 
 namespace ParametricTypeHash
 {
 
-static uint32 Calculate(FOptionalConcreteTypenameId Name, TConstArrayView<FTypeId> Parameters)
+static uint32 Calculate(FOptionalConcreteTypenameId Name, TConstArrayView<FType> Parameters)
 {
-	uint64 ParametersHash = FXxHash64::HashBuffer(Parameters.GetData(), sizeof(FTypeId) * Parameters.Num()).Hash;
+	uint64 ParametersHash = FXxHash64::HashBuffer(Parameters.GetData(), sizeof(FType) * Parameters.Num()).Hash;
 	return HashCombineFast(GetTypeHash(Name), static_cast<uint32>(ParametersHash));
 }
 
 static constexpr uint8 FreeSlotByte = 0xFF;
 static constexpr uint32 FreeSlot = 0xFFFFFFFF;
 
-static uint32* Rehash(uint32 NumSlots, TConstArrayView<FParametricType> Types, TConstArrayView<FTypeId> Parameters)
+static uint32* Rehash(uint32 NumSlots, TConstArrayView<FParametricType> Types, TConstArrayView<FType> Parameters)
 {
 	uint32* Slots = (uint32*)FMemory::Malloc(NumSlots * sizeof(uint32));
 	FMemory::Memset(Slots, FreeSlotByte, NumSlots * sizeof(uint32));
@@ -102,7 +102,7 @@ FParametricTypeId FParametricTypeIndexer::Index(FParametricTypeView View)
 			FParametricType Existing = Types[Slot];
 			if (View.Name == Existing.Name &&
 				View.NumParameters == Existing.Parameters.NumParameters &&
-				!FMemory::Memcmp(View.Parameters, &Parameters[Existing.Parameters.Idx], View.NumParameters * sizeof(FTypeId)))
+				!FMemory::Memcmp(View.Parameters, &Parameters[Existing.Parameters.Idx], View.NumParameters * sizeof(FType)))
 			{
 				return FParametricTypeId(View.NumParameters, Slot);
 			}
@@ -139,38 +139,38 @@ FScopeId FIdIndexerBase::NestReversedScopes(TConstArrayView<FFlatScopeId> Inners
 	return Out;
 }
 
-FParametricTypeId FIdIndexerBase::MakeParametricTypeId(FOptionalConcreteTypenameId Name, TConstArrayView<FTypeId> Params)
+FParametricTypeId FIdIndexerBase::MakeParametricTypeId(FOptionalConcreteTypenameId Name, TConstArrayView<FType> Params)
 {
 	return ParametricTypes.Index({Name, IntCastChecked<uint8>(Params.Num()), Params.GetData()});
 }
 
-FTypeId FIdIndexerBase::MakeParametricType(FTypeId Type, TConstArrayView<FTypeId> Params)
+FType FIdIndexerBase::MakeParametricType(FType Type, TConstArrayView<FType> Params)
 {
 	return {Type.Scope, FTypenameId(MakeParametricTypeId(Type.Name.AsConcrete(), Params))};
 }
 
-FTypeId FIdIndexerBase::MakeAnonymousParametricType(TConstArrayView<FTypeId> Params)
+FType FIdIndexerBase::MakeAnonymousParametricType(TConstArrayView<FType> Params)
 {
 	return { NoId, FTypenameId(MakeParametricTypeId(NoId, Params)) };
 }
 
-FEnumSchemaId FIdIndexerBase::IndexEnum(FTypeId Type)
+FEnumSchemaId FIdIndexerBase::IndexEnum(FType Type)
 {
 	return { IntCastChecked<uint32>(Enums.Add(Type).AsInteger()) };
 }
 
-FStructSchemaId	FIdIndexerBase::IndexStruct(FTypeId Type)
+FStructSchemaId	FIdIndexerBase::IndexStruct(FType Type)
 {
 	return { IntCastChecked<uint32>(Structs.Add(Type).AsInteger()) };
 }
 
-FTypeId FIdIndexerBase::Resolve(FEnumSchemaId Id) const
+FType FIdIndexerBase::Resolve(FEnumSchemaId Id) const
 {
 	check(Id.Idx < uint32(Enums.Num()));
 	return Enums.Get(FSetElementId::FromInteger(Id.Idx));
 }
 
-FTypeId FIdIndexerBase::Resolve(FStructSchemaId Id) const
+FType FIdIndexerBase::Resolve(FStructSchemaId Id) const
 {
 	check(Id.Idx < uint32(Structs.Num()));
 	return Structs.Get(FSetElementId::FromInteger(Id.Idx));
diff --git a/Engine/Plugins/Experimental/PlainProps/Source/Private/PlainPropsInternalFormat.h b/Engine/Plugins/Experimental/PlainProps/Source/Private/PlainPropsInternalFormat.h
--- a/Engine/Plugins/Experimental/PlainProps/Source/Private/PlainPropsInternalFormat.h
+++ b/Engine/Plugins/Experimental/PlainProps/Source/Private/PlainPropsInternalFormat.h
@@ -33,9 +33,9 @@ struct FSchemaBatch
 		return MakeArrayView(reinterpret_cast<const FParametricType*>(GetNestedScopes().end()), NumParametricTypes);
 	}
 
-	const FTypeId* GetFirstParameter() const
+	const FType* GetFirstParameter() const
 	{
-		return reinterpret_cast<const FTypeId*>(GetParametricTypes().end());
+		return reinterpret_cast<const FType*>(GetParametricTypes().end());
 	}
 
 	void ValidateBounds(uint64 NumBytes) const;
@@ -63,7 +63,7 @@ inline bool SkipDeclaredSuperSchema(ESuper Inheritance)	{ return Inheritance ==
 
 struct FStructSchema
 {
-	FTypeId Type;
+	FType Type;
 	uint16 NumMembers;
 	uint16 NumRangeTypes;
 	uint16 NumInnerSchemas;
@@ -102,8 +102,8 @@ struct FStructSchema
 
 static_assert(sizeof(FStructSchema) == 16 && alignof(FStructSchema) == 4,		"Add binary format versioning and read old format in AllocateReadSchemas");
 static_assert(offsetof(FStructSchema, Type) == 0,								"Add binary format versioning and read old format in AllocateReadSchemas");
-static_assert(offsetof(FTypeId, Scope) == 0,									"Add binary format versioning and read old format in AllocateReadSchemas");
-static_assert(offsetof(FTypeId, Name) == 4,										"Add binary format versioning and read old format in AllocateReadSchemas");
+static_assert(offsetof(FType, Scope) == 0,									"Add binary format versioning and read old format in AllocateReadSchemas");
+static_assert(offsetof(FType, Name) == 4,										"Add binary format versioning and read old format in AllocateReadSchemas");
 static_assert(offsetof(FStructSchema, NumMembers) == 8,							"Add binary format versioning and read old format in AllocateReadSchemas");
 static_assert(offsetof(FStructSchema, NumRangeTypes) == 10,						"Add binary format versioning and read old format in AllocateReadSchemas");
 static_assert(offsetof(FStructSchema, NumInnerSchemas) == 12,					"Add binary format versioning and read old format in AllocateReadSchemas");
@@ -122,7 +122,7 @@ inline uint32 CalculateSize(const FStructSchema& In)
 
 struct FEnumSchema
 {
-	FTypeId Type;
+	FType Type;
 	uint8 FlagMode : 1;
 	uint8 ExplicitConstants : 1; // 0,1,2,3,4.. or 1,2,4,8..
 	ELeafWidth Width;
diff --git a/Engine/Plugins/Experimental/PlainProps/Source/Private/PlainPropsInternalParse.h b/Engine/Plugins/Experimental/PlainProps/Source/Private/PlainPropsInternalParse.h
--- a/Engine/Plugins/Experimental/PlainProps/Source/Private/PlainPropsInternalParse.h
+++ b/Engine/Plugins/Experimental/PlainProps/Source/Private/PlainPropsInternalParse.h
@@ -97,7 +97,7 @@ struct FParsedSchemas
 class FTextIndexer : public TIdIndexer<FUtf8StringView>
 {
 public:
-	FOptionalStructSchemaId GetStructId(FTypeId Type) const
+	FOptionalStructSchemaId GetStructId(FType Type) const
 	{
 		FSetElementId Id = Structs.FindId(Type);
 		if (Id.IsValidId())
@@ -107,7 +107,7 @@ public:
 		return NoId;
 	}
 
-	FOptionalEnumSchemaId GetEnumId(FTypeId Type) const
+	FOptionalEnumSchemaId GetEnumId(FType Type) const
 	{
 		FSetElementId Id = Enums.FindId(Type);
 		if (Id.IsValidId())
@@ -129,7 +129,7 @@ struct FTypeTokens
 
 TOptional<FTypeTokens> TokenizeType(FUtf8StringView String);
 
-FTypeId MakeType(const FTypeTokens& TypeTokens, FTextIndexer& Names);
+FType MakeType(const FTypeTokens& TypeTokens, FTextIndexer& Names);
 
 ///////////////////////////////////////////////////////////////////////////////
 
diff --git a/Engine/Plugins/Experimental/PlainProps/Source/Private/PlainPropsParse.cpp b/Engine/Plugins/Experimental/PlainProps/Source/Private/PlainPropsParse.cpp
--- a/Engine/Plugins/Experimental/PlainProps/Source/Private/PlainPropsParse.cpp
+++ b/Engine/Plugins/Experimental/PlainProps/Source/Private/PlainPropsParse.cpp
@@ -329,10 +329,10 @@ TOptional<FTypeTokens> TokenizeType(FUtf8StringView String)
 	return NullOpt;
 }
 
-FTypeId MakeType(const FTypeTokens& TypeTokens, FTextIndexer& Names)
+FType MakeType(const FTypeTokens& TypeTokens, FTextIndexer& Names)
 {
 	// build parameter array for parametric types
-	TArray<FTypeId> Params;
+	TArray<FType> Params;
 	if (int32 Num = TypeTokens.Parameters.Num())
 	{
 		Params.Reserve(Num);
@@ -951,7 +951,7 @@ FParsedEnumSchema FEnumSchemaParser::DeclareEnumSchema(
 	TConstArrayView<FUtf8StringView> ConstantNames,
 	TConstArrayView<uint64> ConstantValues)
 {
-	FTypeId Type = MakeType(TypeTokens, Names);
+	FType Type = MakeType(TypeTokens, Names);
 	FEnumSchemaId Id = Names.IndexEnum(Type);
 
 	TArray<FEnumerator, TInlineAllocator<64>> Enumerators;
@@ -1044,7 +1044,7 @@ FParsedStructSchema FStructSchemaParser::DeclareStructSchema(
 	TConstArrayView<FUtf8StringView> MemberNames,
 	TConstArrayView<FUtf8StringView> MemberTypes)
 {
-	const FTypeId Type = MakeType(TypeTokens, Names);
+	const FType Type = MakeType(TypeTokens, Names);
 	const FStructSchemaId Id = Names.IndexStruct(Type);
 	const FOptionalStructSchemaId DeclaredSuper = ParsedSuper ?
 		Names.IndexStruct(MakeType(ParsedSuper.GetValue(), Names)) :
@@ -1064,7 +1064,7 @@ FParsedStructSchema FStructSchemaParser::DeclareStructSchema(
 		{
 			TOptional<FTypeTokens> ParsedSuperMember = TokenizeType(MemberTypes[0]);
 			check(ParsedSuperMember);
-			FTypeId UsedSuperType = MakeType(ParsedSuperMember.GetValue(), Names);
+			FType UsedSuperType = MakeType(ParsedSuperMember.GetValue(), Names);
 			UsedSuper = Names.IndexStruct(UsedSuperType);
 		}
 	}
@@ -1130,7 +1130,7 @@ bool FStructSchemaParser::ParseMemberSchema(FParsedMemberSchema& Out, FUtf8Strin
 	}
 	else if (TOptional<FTypeTokens> TypeTokens = TokenizeType(TypeOrLeaf))
 	{
-		FTypeId Type = MakeType(TypeTokens.GetValue(), Names);
+		FType Type = MakeType(TypeTokens.GetValue(), Names);
 		FOptionalEnumSchemaId Enum = Names.GetEnumId(Type);
 		if (Enum)
 		{
@@ -1455,7 +1455,7 @@ bool FMemberParser::ParseDynamicStructType(FOptionalStructSchemaId& Out)
 	const FUtf8StringView StructName = Token.Value();
 	TOptional<FTypeTokens> TypeTokens = TokenizeType(StructName);
 	check(TypeTokens);
-	FTypeId Type = MakeType(TypeTokens.GetValue(), Names);
+	FType Type = MakeType(TypeTokens.GetValue(), Names);
 	FOptionalStructSchemaId Schema = Names.GetStructId(Type); 
 	if (Schema)
 	{
@@ -1848,7 +1848,7 @@ void FBatchParser::ParseObjects(TConstArrayView64<FToken> TokensView)
 			continue;
 		}
 
-		const FTypeId Type = MakeType(TypeTokens.GetValue(), Names);
+		const FType Type = MakeType(TypeTokens.GetValue(), Names);
 		const FOptionalStructSchemaId StructId = Names.GetStructId(Type);
 
 		checkf(StructId, TEXT("Object struct schema '%s' not found"), *Print(Token));
diff --git a/Engine/Plugins/Experimental/PlainProps/Source/Private/PlainPropsPrint.cpp b/Engine/Plugins/Experimental/PlainProps/Source/Private/PlainPropsPrint.cpp
--- a/Engine/Plugins/Experimental/PlainProps/Source/Private/PlainPropsPrint.cpp
+++ b/Engine/Plugins/Experimental/PlainProps/Source/Private/PlainPropsPrint.cpp
@@ -138,7 +138,7 @@ class FStructSchemaReader
 public:
 	FStructSchemaReader(const FStructSchema& Schema, FReadBatchId InBatch);
 
-	FTypeId					GetTypeId() const		{ return TypeId; }
+	FType					GetStruct() const		{ return Struct; }
 	bool					IsDense() const			{ return bIsDense; }
 	bool					HasSuper() const		{ return bHasSuper; }
 
@@ -154,7 +154,7 @@ public:
 private:
 	const FMemberType*		Footer;
 	const FReadBatchId		Batch;					// Needed to resolve schemas
-	const FTypeId			TypeId;
+	const FType				Struct;
 	const bool				bIsDense : 1;
 	const bool				bHasSuper : 1;
 	const bool				bUsesSuper : 1;
@@ -281,7 +281,7 @@ void PrintYamlBatch(
 FStructSchemaReader::FStructSchemaReader(const FStructSchema& Schema, FReadBatchId InBatch)
 : Footer(Schema.Footer)
 , Batch(InBatch)
-, TypeId(Schema.Type)
+, Struct(Schema.Type)
 , bIsDense(Schema.IsDense)
 , bHasSuper(Schema.Inheritance != ESuper::No)
 , bUsesSuper(UsesSuper(Schema.Inheritance))
@@ -355,14 +355,14 @@ FSchemaId FStructSchemaReader::GrabInnerSchema()
 	return GetInnerSchemas()[Idx];
 }
 
-FOptionalSchemaId FStructSchemaReader::GrabLeafSchema(FLeafType Leaf)
+FOptionalSchemaId FStructSchemaReader::GrabLeafSchema(FLeafType Member)
 {
-	return Leaf.Type == ELeafType::Enum ? ToOptional(GrabInnerSchema()) : NoId;
+	return Member.Type == ELeafType::Enum ? ToOptional(GrabInnerSchema()) : NoId;
 }
 
-FOptionalSchemaId FStructSchemaReader::GrabStructSchema(FStructType Struct)
+FOptionalSchemaId FStructSchemaReader::GrabStructSchema(FStructType Member)
 {
-	return Struct.IsDynamic ? NoId : ToOptional(GrabInnerSchema());
+	return Member.IsDynamic ? NoId : ToOptional(GrabInnerSchema());
 }
 
 FOptionalSchemaId FStructSchemaReader::GrabRangeSchema(FMemberType InnermostType)
@@ -621,7 +621,7 @@ void FBatchPrinter::PrintStructSchema(const FStructSchema& Struct, FReadBatchId
 {
 	FStructSchemaReader Reader(Struct, BatchId);
 
-	TextBuilder.BeginStruct(*FPrintId<128>(Ids, Reader.GetTypeId()));
+	TextBuilder.BeginStruct(*FPrintId<128>(Ids, Reader.GetStruct()));
 
 	if (Reader.HasSuper())
 	{
@@ -1096,7 +1096,7 @@ void FIds::AppendString(FUtf8Builder& Out, FTypenameId Typename) const
 	else
 	{
 		FParametricTypeView ParametricType = Resolve(Typename.AsParametric());
-		TConstArrayView<FTypeId> Parameters = ParametricType.GetParameters();
+		TConstArrayView<FType> Parameters = ParametricType.GetParameters();
 	
 		if (ParametricType.Name)
 		{
@@ -1104,7 +1104,7 @@ void FIds::AppendString(FUtf8Builder& Out, FTypenameId Typename) const
 		}
 
 		Out.AppendChar(ParametricType.Name ? '<' : '[');
-		for (FTypeId Parameter : Parameters.LeftChop(1))
+		for (FType Parameter : Parameters.LeftChop(1))
 		{
 			AppendString(Out, Parameter);
 			Out.AppendChar(',');
@@ -1117,7 +1117,7 @@ void FIds::AppendString(FUtf8Builder& Out, FTypenameId Typename) const
 	}
 }
 
-void FIds::AppendString(FUtf8Builder& Out, FTypeId Type) const
+void FIds::AppendString(FUtf8Builder& Out, FType Type) const
 {
 	if (Type.Scope)
 	{
@@ -1221,7 +1221,7 @@ FString FDebugIds::Print(FParametricTypeId Typename) const
 	return Print(FTypenameId(Typename));
 }
 
-FString FDebugIds::Print(FTypeId Type) const
+FString FDebugIds::Print(FType Type) const
 {
 	TUtf8StringBuilder<128> Out;
 	Ids.AppendString(Out, Type);
diff --git a/Engine/Plugins/Experimental/PlainProps/Source/Private/PlainPropsWrite.cpp b/Engine/Plugins/Experimental/PlainProps/Source/Private/PlainPropsWrite.cpp
--- a/Engine/Plugins/Experimental/PlainProps/Source/Private/PlainPropsWrite.cpp
+++ b/Engine/Plugins/Experimental/PlainProps/Source/Private/PlainPropsWrite.cpp
@@ -34,7 +34,7 @@ struct FWriteIds
 	uint32								NumKeptStructSchemas;
 	TArray<FNestedScope>				KeptScopes;
 	TArray<FParametricType>				KeptParametrics;
-	TArray<FTypeId>						KeptParameters;
+	TArray<FType>						KeptParameters;
 	TArray<FNameId>						KeptNames;
 
 	FDebugIds							Debug;
@@ -47,7 +47,7 @@ struct FWriteIds
 	FConcreteTypenameId					Remap(FConcreteTypenameId Old) const	{ return { Remap(Old.Id) }; }	
 	FParametricTypeId					Remap(FParametricTypeId Old) const		{ return ParametricTypes[Old.Idx].Get(); }
 	FTypenameId							Remap(FTypenameId Old) const			{ return Old.IsConcrete() ? FTypenameId(Remap(Old.AsConcrete())) : FTypenameId(Remap(Old.AsParametric())); }
-	FTypeId								Remap(FTypeId Old) const				{ return { Remap(Old.Scope), Remap(Old.Name) }; }
+	FType								Remap(FType Old) const				{ return { Remap(Old.Scope), Remap(Old.Name) }; }
 	FStructSchemaId						RemapStruct(FSchemaId Old) const;
 	FEnumSchemaId						RemapEnum(FSchemaId Old) const			{ return Enums[Old.Idx].Get(); }
 
@@ -120,7 +120,7 @@ struct FUsedIds
 		}
 	}
 	
-	void MarkUsed(FTypeId Type)
+	void MarkUsed(FType Type)
 	{
 		MarkUsed(Type.Scope);
 		MarkUsed(Type.Name);
@@ -162,7 +162,7 @@ struct FUsedIds
 
 				FParametricTypeView ParametricType = Ids.Resolve(Typename.AsParametric());
 				MarkUsed(ParametricType.Name);
-				for (FTypeId Parameter : ParametricType.GetParameters())
+				for (FType Parameter : ParametricType.GetParameters())
 				{
 					MarkUsed(Parameter);
 				}
@@ -270,7 +270,7 @@ FWriteIds::FWriteIds(const FIdIndexerBase& Ids, const IStructBindIds& InBindIds,
 		CopyUsed(/* out */ KeptParametrics, Used.KeptParametricTypes, Ids.GetParametricTypes());
 
 		// Remap copied parametric types and copy parameters
-		TConstArrayView<FTypeId> AllParameters = Ids.GetParametricTypes().GetAllParameters();
+		TConstArrayView<FType> AllParameters = Ids.GetParametricTypes().GetAllParameters();
 		KeptParameters.Reserve(SumParams);
 		for (FParametricType& KeptType : KeptParametrics)
 		{
@@ -281,7 +281,7 @@ FWriteIds::FWriteIds(const FIdIndexerBase& Ids, const IStructBindIds& InBindIds,
 		}
 
 		// Remap copied parameters
-		for (FTypeId& KeptParameter : KeptParameters)
+		for (FType& KeptParameter : KeptParameters)
 		{
 			KeptParameter = Remap(KeptParameter);
 		}
diff --git a/Engine/Plugins/Experimental/PlainProps/Source/Public/PlainPropsBind.h b/Engine/Plugins/Experimental/PlainProps/Source/Public/PlainPropsBind.h
--- a/Engine/Plugins/Experimental/PlainProps/Source/Public/PlainPropsBind.h
+++ b/Engine/Plugins/Experimental/PlainProps/Source/Public/PlainPropsBind.h
@@ -575,10 +575,10 @@ using RangeBind = typename TRangeBind<T>::Type;
 
 ////////////////////////////////////////////////////////////////////////////////////////////////
 
-struct FDualTypeId
+struct FDualType
 {
-	FTypeId		BindType;
-	FTypeId		DeclType;
+	FType		BindType;
+	FType		DeclType;
 
 	bool IsLowered() const { return BindType != DeclType; }
 };
@@ -598,7 +598,7 @@ struct FMemberBinding
 	TConstArrayView<FRangeBinding>	RangeBindings;		// Non-empty -> Range
 
 	// @pre InnermostSchema isn't type-erased / lowered
-	PLAINPROPS_API FDualTypeId		IndexParameterName(FIdIndexerBase& Ids) const;
+	PLAINPROPS_API FDualType		IndexParameterName(FIdIndexerBase& Ids) const;
 };
 
 class FSchemaBindings : public IStructBindIds
@@ -672,9 +672,9 @@ constexpr std::string_view SelectStructName()
 }
 
 template<class Ids, ETypename Kind, class Typename>
-FTypeId IndexStructName()
+FType IndexStructName()
 {
-	FTypeId BaseName = { IndexNamespaceId<Ids, Typename>(), 
+	FType BaseName = { IndexNamespaceId<Ids, Typename>(), 
 						 FTypenameId(Ids::IndexTypename(ToAnsiView(SelectStructName<Kind, Typename>()))) };
 	
 	if constexpr (ParametricName<Typename>)
@@ -702,7 +702,7 @@ FStructSchemaId IndexStructBindIdIfNeeded(FStructSchemaId DeclId)
 }
 
 template<class Ids, class Typename>
-FDualStructSchemaId IndexStructDualId(FTypeId DeclName = IndexStructName<Ids, ETypename::Decl, Typename>())
+FDualStructSchemaId IndexStructDualId(FType DeclName = IndexStructName<Ids, ETypename::Decl, Typename>())
 {	
 	FDualStructSchemaId Out;
 	Out.DeclId = Ids::IndexStruct(DeclName);
@@ -710,7 +710,7 @@ FDualStructSchemaId IndexStructDualId(FTypeId DeclName = IndexStructName<Ids, ET
 
 	if constexpr (ExplicitBindName<Typename> || ParametricName<Typename>)
 	{
-		FTypeId BindName = IndexStructName<Ids, ETypename::Bind, Typename>();
+		FType BindName = IndexStructName<Ids, ETypename::Bind, Typename>();
 		Out.BindId = BindName != DeclName ? Ids::IndexStruct(BindName) : Out.DeclId;
 	}
 
@@ -742,7 +742,7 @@ FDualStructSchemaId GetStructDualId()
 }
 
 template<class Ids, typename Ctti>
-FTypeId IndexCttiName()
+FType IndexCttiName()
 {
 	FTypenameId Name{Ids::IndexTypename(Ctti::Name)};
 	FScopeId Namespace = NoId;
@@ -762,26 +762,26 @@ FEnumSchemaId GetEnumId()
 }
 
 template<class Ids, Arithmetic T>
-FTypeId IndexArithmeticName()
+FType IndexArithmeticName()
 {
 	static constexpr FUnpackedLeafType Leaf = ReflectArithmetic<T>;
 	return { NoId, FTypenameId(Ids::IndexTypename(ToAnsiView(ArithmeticName<Leaf.Type, Leaf.Width>))) };
 }
 
 template<class Ids, ETypename, Arithmetic Leaf>
-FTypeId IndexParameterName()
+FType IndexParameterName()
 {
 	return IndexArithmeticName<Ids, Leaf>();
 }
 
 template<class Ids, ETypename, Enumeration Enum>
-FTypeId IndexParameterName()
+FType IndexParameterName()
 {
 	return GetEnumId<Ids, Enum>();
 }
 
 template<class Ids, ETypename Kind, typename T>
-FTypeId IndexParameterName()
+FType IndexParameterName()
 {
 	using RangeBinding = RangeBind<T>;
 	if constexpr (std::is_void_v<RangeBinding>)
@@ -790,8 +790,8 @@ FTypeId IndexParameterName()
 	}
 	else
 	{
-		FTypeId ItemParam = IndexParameterName<Ids, Kind, typename RangeBinding::ItemType>();
-		FTypeId SizeParam = Ids::GetIndexer().MakeRangeParameter(RangeSizeOf(typename RangeBinding::SizeType{}));
+		FType ItemParam = IndexParameterName<Ids, Kind, typename RangeBinding::ItemType>();
+		FType SizeParam = Ids::GetIndexer().MakeRangeParameter(RangeSizeOf(typename RangeBinding::SizeType{}));
 
 		if constexpr (Kind == ETypename::Decl)
 		{
@@ -801,16 +801,16 @@ FTypeId IndexParameterName()
 		else
 		{
 			using Typename = TTypename<T>;
-			FTypeId RangeBindName = { IndexNamespaceId<Ids, Typename>(), FTypenameId(Ids::IndexTypename(ToAnsiView(RangeBinding::BindName))) };
+			FType RangeBindName = { IndexNamespaceId<Ids, Typename>(), FTypenameId(Ids::IndexTypename(ToAnsiView(RangeBinding::BindName))) };
 			return Ids::GetIndexer().MakeParametricType(RangeBindName, {ItemParam, SizeParam});
 		}
 	}
 }
 
 template<class Ids, ETypename Kind, typename... Ts>
-FTypeId IndexParametricType(FTypeId TemplatedType, const std::tuple<Ts...>*)
+FType IndexParametricType(FType TemplatedType, const std::tuple<Ts...>*)
 {
-	FTypeId Parameters[] = { (IndexParameterName<Ids, Kind, Ts>())... };
+	FType Parameters[] = { (IndexParameterName<Ids, Kind, Ts>())... };
 	return Ids::GetIndexer().MakeParametricType(TemplatedType, Parameters);
 }
 
@@ -857,7 +857,7 @@ struct TScopedStructBinding : FDualStructSchemaId, CustomBinding
 	using Ids = typename Runtime::Ids;
 	using Typename = TCustomTypename<CustomBinding>::Type;
 
-	explicit TScopedStructBinding(FTypeId DeclName = IndexStructName<Ids, ETypename::Decl, Typename>(), InnerStructArray Lowered = {})
+	explicit TScopedStructBinding(FType DeclName = IndexStructName<Ids, ETypename::Decl, Typename>(), InnerStructArray Lowered = {})
 	: TScopedStructBinding(DeclName, IndexStructDualId<Ids, Typename>(DeclName), Lowered)
 	{}
 
@@ -869,7 +869,7 @@ struct TScopedStructBinding : FDualStructSchemaId, CustomBinding
 	: TScopedStructBinding(Ids::GetIndexer().Resolve(Id), {Id, Id}, Lowered)
 	{}
 
-	TScopedStructBinding(FTypeId DeclName, FDualStructSchemaId DualId, InnerStructArray& Lowered)
+	TScopedStructBinding(FType DeclName, FDualStructSchemaId DualId, InnerStructArray& Lowered)
 	: FDualStructSchemaId(DualId)
 	, CustomBinding(TCustomInit<Ids>{{Lowered}})
 	{
@@ -909,7 +909,7 @@ struct TScopedStructBinding<T, Runtime, void>
 	const FStructSchemaId DeclId;
 	const FStructSchemaId BindId;
 
-	TScopedStructBinding(FTypeId DeclName = IndexStructName<Ids, ETypename::Decl, Typename>())
+	TScopedStructBinding(FType DeclName = IndexStructName<Ids, ETypename::Decl, Typename>())
 	: DeclId(DeclareNativeStruct<CttiOf<T>, Ids>(Runtime::GetTypes(), TOccupancyOf<T>::Value))
 	, BindId(IndexStructBindIdIfNeeded<Ids, Typename>(DeclId))
 	{
@@ -1083,7 +1083,7 @@ FEnumSchemaId DeclareNativeEnum(FDeclarations& Out, EEnumMode Mode)
 {
 	using UnderlyingType = std::underlying_type_t<typename Ctti::Type>;
 
-	FTypeId Type = IndexCttiName<Ids, Ctti>();
+	FType Type = IndexCttiName<Ids, Ctti>();
 	FEnumSchemaId Id = Ids::IndexEnum(Type);
 	FEnumerator Enumerators[Ctti::NumEnumerators];
 	for (FEnumerator& Enumerator : Enumerators)
@@ -1102,7 +1102,7 @@ FStructSchemaId DeclareNativeStruct(FDeclarations& Out, EMemberPresence Occupanc
 	using Typename = TTypename<typename Ctti::Type>;
 	using SuperType = typename Ctti::Super;
 
-	FTypeId Type = IndexStructName<Ids, ETypename::Decl, Typename>();
+	FType Type = IndexStructName<Ids, ETypename::Decl, Typename>();
 	FStructSchemaId Id = Ids::IndexStruct(Type);
 	FOptionalStructSchemaId SuperId;
 	if constexpr (!std::is_void_v<SuperType>)
@@ -1152,7 +1152,7 @@ struct FIdBinding
 	FConcreteTypenameId					Remap(FConcreteTypenameId Old) const	{ return { Remap(Old.Id) }; }	
 	FParametricTypeId					Remap(FParametricTypeId Old) const		{ return ParametricTypes[Old.Idx]; }
 	FTypenameId							Remap(FTypenameId Old) const			{ return Old.IsConcrete() ? FTypenameId(Remap(Old.AsConcrete())) : FTypenameId(Remap(Old.AsParametric())); }
-	FTypeId								Remap(FTypeId Old) const				{ return { Remap(Old.Scope), Remap(Old.Name) }; }
+	FType								Remap(FType Old) const				{ return { Remap(Old.Scope), Remap(Old.Name) }; }
 
 	template<typename T>
 	TOptionalId<T>						Remap(TOptionalId<T> Old) const			{ return Old ? ToOptional(Remap(Old.Get())) : Old; }
diff --git a/Engine/Plugins/Experimental/PlainProps/Source/Public/PlainPropsBuildSchema.h b/Engine/Plugins/Experimental/PlainProps/Source/Public/PlainPropsBuildSchema.h
--- a/Engine/Plugins/Experimental/PlainProps/Source/Public/PlainPropsBuildSchema.h
+++ b/Engine/Plugins/Experimental/PlainProps/Source/Public/PlainPropsBuildSchema.h
@@ -17,7 +17,7 @@ struct FStructSchemaBuilder;
 
 struct FBuiltStructSchema
 {
-	FTypeId							Type;
+	FType							Type;
 	FStructSchemaId					Id;
 	FOptionalStructSchemaId			Super;
 	bool							bDense = false;
@@ -27,7 +27,7 @@ struct FBuiltStructSchema
 
 struct FBuiltEnumSchema
 {
-	FTypeId							Type;
+	FType							Type;
 	FEnumSchemaId					Id;
 	EEnumMode						Mode = EEnumMode::Flat;
 	ELeafWidth						Width = ELeafWidth::B8;
diff --git a/Engine/Plugins/Experimental/PlainProps/Source/Public/PlainPropsDeclare.h b/Engine/Plugins/Experimental/PlainProps/Source/Public/PlainPropsDeclare.h
--- a/Engine/Plugins/Experimental/PlainProps/Source/Public/PlainPropsDeclare.h
+++ b/Engine/Plugins/Experimental/PlainProps/Source/Public/PlainPropsDeclare.h
@@ -20,7 +20,7 @@ enum class EEnumMode { Flat, Flag };
 
 struct FEnumDeclaration
 {
-	FTypeId					Type;			// Could be removed
+	FType					Type;			// Could be removed
 	EEnumMode				Mode;
 	//EOptionalLeafWidth	Width;			// Possible approach to declare enums w/o any noted values
 	uint16					NumEnumerators;
@@ -38,7 +38,7 @@ struct FStructDeclaration
 {
 	uint32					RefCount;
 	FStructSchemaId			Id;				// Could be removed, might allow declaration dedup among templated types
-	FTypeId					Type;			// Could be removed, might allow declaration dedup among templated types
+	FType					Type;			// Could be removed, might allow declaration dedup among templated types
 	FOptionalStructSchemaId	Super;
 	EMemberPresence			Occupancy;
 	uint16					NumMembers;
@@ -62,10 +62,10 @@ public:
 	UE_NONCOPYABLE(FDeclarations);
 	explicit FDeclarations(FDebugIds In) : Debug(In) {}
 
-	PLAINPROPS_API const FEnumDeclaration&			DeclareEnum(FEnumSchemaId DeclId, FTypeId Type, EEnumMode Mode, TConstArrayView<FEnumerator> Enumerators, EEnumAliases Policy);
+	PLAINPROPS_API const FEnumDeclaration&			DeclareEnum(FEnumSchemaId DeclId, FType Type, EEnumMode Mode, TConstArrayView<FEnumerator> Enumerators, EEnumAliases Policy);
 	// Declare struct with ref count 1 or increment it and check that previous declaration matches
-	PLAINPROPS_API const FStructDeclaration&		DeclareStruct(FStructSchemaId DeclId, FTypeId Type, TConstArrayView<FMemberId> MemberOrder, EMemberPresence Occupancy, FOptionalStructSchemaId Super = {});
-	PLAINPROPS_API const FStructDeclaration&		DeclareNumeralStruct(FStructSchemaId DeclId, FTypeId Type, TConstArrayView<FMemberId> Numerals, EMemberPresence Occupancy);
+	PLAINPROPS_API const FStructDeclaration&		DeclareStruct(FStructSchemaId DeclId, FType Type, TConstArrayView<FMemberId> MemberOrder, EMemberPresence Occupancy, FOptionalStructSchemaId Super = {});
+	PLAINPROPS_API const FStructDeclaration&		DeclareNumeralStruct(FStructSchemaId DeclId, FType Type, TConstArrayView<FMemberId> Numerals, EMemberPresence Occupancy);
 	
 	void											DropEnum(FEnumSchemaId DeclId)			{ Check(DeclId); DeclaredEnums[DeclId.Idx].Reset(); }
 	PLAINPROPS_API void								DropStructRef(FStructSchemaId DeclId);
diff --git a/Engine/Plugins/Experimental/PlainProps/Source/Public/PlainPropsIndex.h b/Engine/Plugins/Experimental/PlainProps/Source/Public/PlainPropsIndex.h
--- a/Engine/Plugins/Experimental/PlainProps/Source/Public/PlainPropsIndex.h
+++ b/Engine/Plugins/Experimental/PlainProps/Source/Public/PlainPropsIndex.h
@@ -35,13 +35,13 @@ public:
 	FParametricType						At(int32 Idx) const									{ return Types[Idx]; }
 	int32								Num() const											{ return Types.Num(); }
 	TConstArrayView<FParametricType>	GetAllTypes() const									{ return Types; }
-	TConstArrayView<FTypeId>			GetAllParameters() const							{ return Parameters; }
+	TConstArrayView<FType>				GetAllParameters() const							{ return Parameters; }
 
 private:
 	uint32								NumSlots = 0;
 	uint32*								Slots = nullptr;
 	TArray<FParametricType>				Types;
-	TArray<FTypeId>						Parameters;
+	TArray<FType>						Parameters;
 };
 
 //////////////////////////////////////////////////////////////////////////
@@ -51,14 +51,14 @@ class FIdIndexerBase : public FIds
 public:
 	PLAINPROPS_API FScopeId				NestFlatScope(FScopeId Outer, FFlatScopeId Inner);
 	PLAINPROPS_API FScopeId				NestReversedScopes(TConstArrayView<FFlatScopeId> Inners);
-	PLAINPROPS_API FParametricTypeId	MakeParametricTypeId(FOptionalConcreteTypenameId Name, TConstArrayView<FTypeId> Params);
-	PLAINPROPS_API FTypeId				MakeParametricType(FTypeId Type, TConstArrayView<FTypeId> Params);
-	PLAINPROPS_API FTypeId				MakeAnonymousParametricType(TConstArrayView<FTypeId> Params);
-	FTypeId								MakeLeafParameter(FUnpackedLeafType Leaf)			{ return { NoId, FTypenameId(Leaves[uint8(Leaf.Type)][uint8(Leaf.Width)]) }; }
-	FTypeId								MakeRangeParameter(ERangeSizeType SizeType)			{ return { NoId, FTypenameId(Ranges[uint8(SizeType)]) }; }
-
-	PLAINPROPS_API FEnumSchemaId		IndexEnum(FTypeId Type);
-	PLAINPROPS_API FStructSchemaId		IndexStruct(FTypeId Type);
+	PLAINPROPS_API FParametricTypeId	MakeParametricTypeId(FOptionalConcreteTypenameId Name, TConstArrayView<FType> Params);
+	PLAINPROPS_API FType				MakeParametricType(FType Type, TConstArrayView<FType> Params);
+	PLAINPROPS_API FType				MakeAnonymousParametricType(TConstArrayView<FType> Params);
+	FType								MakeLeafParameter(FUnpackedLeafType Leaf)			{ return { NoId, FTypenameId(Leaves[uint8(Leaf.Type)][uint8(Leaf.Width)]) }; }
+	FType								MakeRangeParameter(ERangeSizeType SizeType)			{ return { NoId, FTypenameId(Ranges[uint8(SizeType)]) }; }
+
+	PLAINPROPS_API FEnumSchemaId		IndexEnum(FType Type);
+	PLAINPROPS_API FStructSchemaId		IndexStruct(FType Type);
 	
 	virtual uint32						NumNestedScopes() const override final				{ return IntCastChecked<uint32>(NestedScopes.Num()); }
 	virtual uint32						NumParametricTypes() const override final			{ return IntCastChecked<uint32>(ParametricTypes.Num()); }
@@ -67,8 +67,8 @@ public:
 	
 	virtual FNestedScope				Resolve(FNestedScopeId Id) const override final		{ return NestedScopes.Resolve(Id); }
 	virtual FParametricTypeView			Resolve(FParametricTypeId Id) const override final	{ return ParametricTypes.Resolve(Id); }
-	PLAINPROPS_API virtual FTypeId		Resolve(FEnumSchemaId Id) const override final;
-	PLAINPROPS_API virtual FTypeId		Resolve(FStructSchemaId Id) const override final;
+	PLAINPROPS_API virtual FType		Resolve(FEnumSchemaId Id) const override final;
+	PLAINPROPS_API virtual FType		Resolve(FStructSchemaId Id) const override final;
 	
 	const FNestedScopeIndexer&			GetNestedScopes() const								{ return NestedScopes; }
 	const FParametricTypeIndexer&		GetParametricTypes() const							{ return ParametricTypes; }
@@ -77,8 +77,8 @@ public:
 protected:
 	FNestedScopeIndexer					NestedScopes;
 	FParametricTypeIndexer				ParametricTypes;
-	TSet<FTypeId>						Enums;
-	TSet<FTypeId>						Structs;
+	TSet<FType>							Enums;
+	TSet<FType>							Structs;
 	FConcreteTypenameId					Leaves[8][4];
 	FConcreteTypenameId					Ranges[9];
 
@@ -102,7 +102,7 @@ public:
 	template<typename T> FScopeId				MakeScope(T&& Name)						{ return FScopeId(NameScope(Forward<T>(Name))); }
 	template<typename T> FScopeId				NestScope(FScopeId Outer, T&& Inner)	{ return NestFlatScope(Outer, NameScope(Forward<T>(Inner))); }
 	template<typename T> FTypenameId			MakeTypename(T&& Name)					{ return FTypenameId(NameType(Forward<T>(Name))); }
-	template<typename T> FTypeId				MakeType(T&& Scope, T&& Name)			{ return { MakeScope(Forward<T>(Scope)), MakeTypename(Forward<T>(Name)) }; }
+	template<typename T> FType					MakeType(T&& Scope, T&& Name)			{ return { MakeScope(Forward<T>(Scope)), MakeTypename(Forward<T>(Name)) }; }
 
 	NameT										ResolveName(FNameId Id)	const			{ return Names.Get(FSetElementId::FromInteger(Id.Idx)); }
 	virtual uint32								NumNames() const override final			{ return static_cast<uint32>(Names.Num()); }
diff --git a/Engine/Plugins/Experimental/PlainProps/Source/Public/PlainPropsRead.h b/Engine/Plugins/Experimental/PlainProps/Source/Public/PlainPropsRead.h
--- a/Engine/Plugins/Experimental/PlainProps/Source/Public/PlainPropsRead.h
+++ b/Engine/Plugins/Experimental/PlainProps/Source/Public/PlainPropsRead.h
@@ -697,7 +697,7 @@ public:
 
 	FMemberId				PeekName() const		{ return It->PeekName().Get(); }
 	EMemberKind				PeekKind() const		{ return It->PeekKind(); }
-	FTypeId					PeekOwner() const		{ return It->Owner; }
+	FType					PeekOwner() const		{ return It->Owner; }
 
 	FLeafView				GrabLeaf()				{ return Grabbed(It->GrabLeaf()); }
 	FRangeView				GrabRange()				{ return Grabbed(It->GrabRange()); }
@@ -707,7 +707,7 @@ private:
 	struct FReader : public FMemberReader
 	{
 		FReader(FStructView Struct);
-		FTypeId Owner;
+		FType Owner;
 	};
 	TArray<FReader, TInlineAllocator<8>> Lineage;
 	FReader* It = nullptr;
diff --git a/Engine/Plugins/Experimental/PlainProps/Source/Public/PlainPropsTypes.h b/Engine/Plugins/Experimental/PlainProps/Source/Public/PlainPropsTypes.h
--- a/Engine/Plugins/Experimental/PlainProps/Source/Public/PlainPropsTypes.h
+++ b/Engine/Plugins/Experimental/PlainProps/Source/Public/PlainPropsTypes.h
@@ -210,13 +210,13 @@ struct FTypenameId : FBaseTypenameId
 
 //////////////////////////////////////////////////////////////////////////
 
-struct FTypeId
+struct FType
 {
 	FScopeId 		Scope;
 	FTypenameId		Name;
 
-	bool operator==(FTypeId O) const { return Scope.AsInt() == O.Scope.AsInt() && Name.AsInt() == O.Name.AsInt(); }
-	friend uint32 GetTypeHash(FTypeId Type) { return HashCombineFast(Type.Scope.AsInt(), Type.Name.Idx); };
+	bool operator==(FType O) const { return Scope.AsInt() == O.Scope.AsInt() && Name.AsInt() == O.Name.AsInt(); }
+	friend uint32 GetTypeHash(FType Type) { return HashCombineFast(Type.Scope.AsInt(), Type.Name.Idx); };
 };
 
 //////////////////////////////////////////////////////////////////////////
@@ -311,16 +311,16 @@ struct FParametricType
 // Fully resolved FParametricTypeId
 struct FParametricTypeView
 {
-	FParametricTypeView(FConcreteTypenameId InName, uint8 NumParams, const FTypeId* Params) : Name(InName), NumParameters(NumParams), Parameters(Params) {}
-	FParametricTypeView(FOptionalConcreteTypenameId InName, uint8 NumParams, const FTypeId* Params) : Name(InName), NumParameters(NumParams), Parameters(Params) {}
-	FParametricTypeView(FConcreteTypenameId InName, TConstArrayView<FTypeId> Params);
-	FParametricTypeView(FOptionalConcreteTypenameId InName, TConstArrayView<FTypeId> Params);
+	FParametricTypeView(FConcreteTypenameId InName, uint8 NumParams, const FType* Params) : Name(InName), NumParameters(NumParams), Parameters(Params) {}
+	FParametricTypeView(FOptionalConcreteTypenameId InName, uint8 NumParams, const FType* Params) : Name(InName), NumParameters(NumParams), Parameters(Params) {}
+	FParametricTypeView(FConcreteTypenameId InName, TConstArrayView<FType> Params);
+	FParametricTypeView(FOptionalConcreteTypenameId InName, TConstArrayView<FType> Params);
 
 	FOptionalConcreteTypenameId	Name;
 	uint8						NumParameters;
-	const FTypeId*				Parameters;
+	const FType*				Parameters;
 
-	TConstArrayView<FTypeId>	GetParameters() const { return MakeArrayView(Parameters, NumParameters); }
+	TConstArrayView<FType>		GetParameters() const { return MakeArrayView(Parameters, NumParameters); }
 };
 
 //////////////////////////////////////////////////////////////////////////
@@ -433,8 +433,8 @@ public:
 
 	virtual FNestedScope				Resolve(FNestedScopeId Id) const = 0;
 	virtual FParametricTypeView			Resolve(FParametricTypeId Id) const = 0;
-	virtual FTypeId						Resolve(FEnumSchemaId Id) const = 0;
-	virtual FTypeId						Resolve(FStructSchemaId Id) const = 0;
+	virtual FType						Resolve(FEnumSchemaId Id) const = 0;
+	virtual FType						Resolve(FStructSchemaId Id) const = 0;
 
 	virtual void						AppendString(FUtf8Builder& Out, FNameId Name) const = 0;
 
@@ -442,7 +442,7 @@ public:
 	PLAINPROPS_API virtual void			AppendString(FUtf8Builder& Out, FOptionalMemberId Name) const;
 	PLAINPROPS_API virtual void			AppendString(FUtf8Builder& Out, FScopeId Scope) const;
 	PLAINPROPS_API virtual void			AppendString(FUtf8Builder& Out, FTypenameId Typename) const;
-	PLAINPROPS_API virtual void			AppendString(FUtf8Builder& Out, FTypeId Type) const;
+	PLAINPROPS_API virtual void			AppendString(FUtf8Builder& Out, FType Type) const;
 	PLAINPROPS_API virtual void			AppendString(FUtf8Builder& Out, FEnumSchemaId Enum) const;
 	PLAINPROPS_API virtual void			AppendString(FUtf8Builder& Out, FStructSchemaId Struct) const;
 };
@@ -459,7 +459,7 @@ public:
 	PLAINPROPS_API FString				Print(FTypenameId Typename) const;
 	PLAINPROPS_API FString				Print(FConcreteTypenameId Typename) const;
 	PLAINPROPS_API FString				Print(FParametricTypeId Typename) const;
-	PLAINPROPS_API FString				Print(FTypeId Type) const;
+	PLAINPROPS_API FString				Print(FType Type) const;
 	PLAINPROPS_API FString				Print(FEnumSchemaId Enum) const;
 	PLAINPROPS_API FString				Print(FStructSchemaId Struct) const;
 private:
diff --git a/Engine/Plugins/Experimental/PlainProps/Source/Public/PlainPropsUeCoreBindings.h b/Engine/Plugins/Experimental/PlainProps/Source/Public/PlainPropsUeCoreBindings.h
--- a/Engine/Plugins/Experimental/PlainProps/Source/Public/PlainPropsUeCoreBindings.h
+++ b/Engine/Plugins/Experimental/PlainProps/Source/Public/PlainPropsUeCoreBindings.h
@@ -49,12 +49,12 @@ namespace PlainProps::UE
 	//template<typename Ctti>
 	//FStructSchemaId			BindStructInterlaced(TConstArrayView<FMemberBinding> NonCttiMembers);
 	//FStructSchemaId			BindStruct(FStructSchemaId Id, const ICustomBinding& Custom);
-	//FStructSchemaId			BindStruct(FTypeId Type, FOptionalSchemaId Super, TConstArrayView<FNamedMemberBinding> Members, EMemberPresence Occupancy);
+	//FStructSchemaId			BindStruct(FType Type, FOptionalSchemaId Super, TConstArrayView<FNamedMemberBinding> Members, EMemberPresence Occupancy);
 	//void					DropStruct(FStructSchemaId Id) { Types.DropStruct(Id); }
 
 	//template<typename Ctti>
 	//FEnumSchemaId			BindEnum();
-	//FEnumSchemaId			BindEnum(FTypeId Type, EEnumMode Mode, ELeafWidth Width, TConstArrayView<FEnumerator> Enumerators);
+	//FEnumSchemaId			BindEnum(FType Type, EEnumMode Mode, ELeafWidth Width, TConstArrayView<FEnumerator> Enumerators);
 	//void					DropEnum(FEnumSchemaId Id) { Types.DropEnum(Id); }
 //};
 //
@@ -65,10 +65,10 @@ namespace PlainProps::UE
 //	static FMemberId		IndexMember(FAnsiStringView Name)			{ return GReflection.GetIds().NameMember(FName(Name)); }
 //	static FTypenameId		IndexTypename(FAnsiStringView Name)			{ return GReflection.GetIds().MakeTypename(FName(Name)); }
 //	static FScopeId			IndexCurrentModule()						{ return GReflection.GetIds().MakeScope(FName(UE_MODULE_NAME)); }
-//	static FTypeId			IndexNativeType(FAnsiStringView Typename)	{ return {IndexCurrentModule(), IndexTypename(Typename)}; }
-//	static FEnumSchemaId	IndexEnum(FTypeId Name)						;
+//	static FType			IndexNativeType(FAnsiStringView Typename)	{ return {IndexCurrentModule(), IndexTypename(Typename)}; }
+//	static FEnumSchemaId	IndexEnum(FType Name)						;
 //	static FEnumSchemaId	IndexEnum(FAnsiStringView Name)				;
-//	static FStructSchemaId	IndexStruct(FTypeId Name)					;
+//	static FStructSchemaId	IndexStruct(FType Name)					;
 //	static FStructSchemaId	IndexStruct(FAnsiStringView Name)			;
 //	
 //};
diff --git a/Engine/Plugins/Experimental/PlainProps/Source/Tests/PlainPropsTests.cpp b/Engine/Plugins/Experimental/PlainProps/Source/Tests/PlainPropsTests.cpp
--- a/Engine/Plugins/Experimental/PlainProps/Source/Tests/PlainPropsTests.cpp
+++ b/Engine/Plugins/Experimental/PlainProps/Source/Tests/PlainPropsTests.cpp
@@ -30,7 +30,7 @@ static bool operator==(FScopeId A, FNestedScopeId B) { return A == FScopeId(B);
 static bool operator==(FScopeId A, FFlatScopeId B) { return A == FScopeId(B); }
 static bool operator==(FParametricTypeView A, FParametricTypeView B)
 {
-	return A.Name == B.Name && A.NumParameters == B.NumParameters && !FMemory::Memcmp(A.Parameters, B.Parameters, A.NumParameters * sizeof(FTypeId));
+	return A.Name == B.Name && A.NumParameters == B.NumParameters && !FMemory::Memcmp(A.Parameters, B.Parameters, A.NumParameters * sizeof(FType));
 }
 
 static uint64 CalculateSize(const FSchemaBatch& In)
@@ -222,9 +222,9 @@ TEST_CASE_NAMED(FPlainPropsIndexTest, "System::Core::Serialization::PlainProps::
 		FConcreteTypenameId T4{{4}};
 		FConcreteTypenameId T5{{5}};
 		
-		FTypeId S0T3 = {S0, FTypenameId{T3}};
-		FTypeId S1T3 = {S1, FTypenameId{T3}};
-		FTypeId S2T3 = {S2, FTypenameId{T3}};
+		FType S0T3 = {S0, FTypenameId{T3}};
+		FType S1T3 = {S1, FTypenameId{T3}};
+		FType S2T3 = {S2, FTypenameId{T3}};
 
 		FParametricTypeIndexer Indexer;
 
@@ -234,8 +234,8 @@ TEST_CASE_NAMED(FPlainPropsIndexTest, "System::Core::Serialization::PlainProps::
 		CHECK(Indexer.Resolve(T4_S0T3) == FParametricTypeView{T4, 1, &S0T3});
 		CHECK(Indexer.Resolve(T4_S1T3) == FParametricTypeView{T4, 1, &S1T3});
 		
-		FTypeId S1T4_S0T3 = {S1,  FTypenameId{T4_S0T3}};
-		FTypeId S2T4_S1T3 = {S2,  FTypenameId{T4_S1T3}};
+		FType S1T4_S0T3 = {S1,  FTypenameId{T4_S0T3}};
+		FType S2T4_S1T3 = {S2,  FTypenameId{T4_S1T3}};
 		
 		CHECK(S1T4_S0T3.Name.AsParametric() == T4_S0T3);
 		CHECK(S2T4_S1T3.Name.AsParametric() == T4_S1T3);
@@ -266,9 +266,9 @@ class FTestBatchBuilder : public TIdIndexer<FAnsiString>, public IStructBindIds
 public:
 	FTestBatchBuilder(FScratchAllocator& InScratch) : Declarations(FDebugIds(*this)), Scratch(InScratch) {}
 
-	FEnumSchemaId				DeclareEnum(FTypeId Type, EEnumMode Mode, ELeafWidth Width, std::initializer_list<const char*> Names, std::initializer_list<uint64> Constants);
+	FEnumSchemaId				DeclareEnum(FType Type, EEnumMode Mode, ELeafWidth Width, std::initializer_list<const char*> Names, std::initializer_list<uint64> Constants);
 	FEnumSchemaId				DeclareEnum(const char* Scope, const char* Name, EEnumMode Mode, ELeafWidth Width, std::initializer_list<const char*> Names, std::initializer_list<uint64> Constants);
-	FStructSchemaId				DeclareStruct(FTypeId Type, std::initializer_list<const char*> MemberOrder, EMemberPresence Occupancy, FOptionalStructSchemaId Super = {});
+	FStructSchemaId				DeclareStruct(FType Type, std::initializer_list<const char*> MemberOrder, EMemberPresence Occupancy, FOptionalStructSchemaId Super = {});
 	FStructSchemaId				DeclareStruct(const char* Scope, const char* Name, std::initializer_list<const char*> MemberOrder, EMemberPresence Occupancy, FOptionalStructSchemaId Super = {});
 	
 	const FEnumDeclaration&		Get(FEnumSchemaId Id) { return Declarations.Get(Id); }
@@ -318,7 +318,7 @@ private:
 	}
 };
 
-FStructSchemaId FTestBatchBuilder::DeclareStruct(FTypeId Type, std::initializer_list<const char*> MemberOrder, EMemberPresence Occupancy, FOptionalStructSchemaId Super)
+FStructSchemaId FTestBatchBuilder::DeclareStruct(FType Type, std::initializer_list<const char*> MemberOrder, EMemberPresence Occupancy, FOptionalStructSchemaId Super)
 {
 	FStructSchemaId Id = IndexStruct(Type);
 	Declarations.DeclareStruct(Id, Type, NameMembers(MemberOrder), Occupancy, Super);
@@ -330,7 +330,7 @@ FStructSchemaId FTestBatchBuilder::DeclareStruct(const char* Scope, const char*
 	return DeclareStruct(MakeType(Scope, Name), MemberOrder, Occupancy, Super);
 }
 
-FEnumSchemaId FTestBatchBuilder::DeclareEnum(FTypeId Type, EEnumMode Mode, ELeafWidth, std::initializer_list<const char*> InNames, std::initializer_list<uint64> Constants)
+FEnumSchemaId FTestBatchBuilder::DeclareEnum(FType Type, EEnumMode Mode, ELeafWidth, std::initializer_list<const char*> InNames, std::initializer_list<uint64> Constants)
 {
 	FEnumSchemaId Id = IndexEnum(Type);
 	Declarations.DeclareEnum(Id, Type, Mode, MakeEnumerators(InNames, Constants), EEnumAliases::Fail);
@@ -430,7 +430,7 @@ public:
 			}
 		}
 		
-		check(Names.Num() >= 3); // 1 FTypeId and 1 member id at least
+		check(Names.Num() >= 3); // 1 FType and 1 member id at least
 		check(NextStr == Data.GetDataEnd()); // end with null-terminator
 	}
 	
@@ -458,8 +458,8 @@ public:
 
 	virtual FNestedScope			Resolve(FNestedScopeId Id) const override final		{ return ResolveNestedScope(Schemas, Id); }
 	virtual FParametricTypeView		Resolve(FParametricTypeId Id) const override final	{ return ResolveParametricType(Schemas, Id); }
-	virtual FTypeId					Resolve(FEnumSchemaId Id) const override final		{ return ResolveEnumSchema(Schemas, Id).Type; }
-	virtual FTypeId					Resolve(FStructSchemaId Id) const override final	{ return ResolveStructSchema(Schemas, Id).Type; }
+	virtual FType					Resolve(FEnumSchemaId Id) const override final		{ return ResolveEnumSchema(Schemas, Id).Type; }
+	virtual FType					Resolve(FStructSchemaId Id) const override final	{ return ResolveStructSchema(Schemas, Id).Type; }
 
 	using FIds::AppendString;
 	virtual void AppendString(FUtf8StringBuilderBase& Out, FNameId Name) const override final { Out.Append(Names[Name]); }
@@ -1610,19 +1610,19 @@ TEST_CASE_NAMED(FPlainPropsReadWriteTest, "System::Core::Serialization::PlainPro
 			FScopeId DoubleNested = Batch.NestScope(NestedUsed, "DoubleNested");
 			FScopeId NestedUnused3 = Batch.NestScope(FlatUsed, "NestedUnused3");
 			
-			FTypeId E1T{NestedUnused1,	Batch.MakeTypename("E1")};
-			FTypeId E2T{NestedUsed,		Batch.MakeTypename("E2")};
-			FTypeId E3T{NestedUnused2,	Batch.MakeTypename("E3")};
+			FType E1T{NestedUnused1,	Batch.MakeTypename("E1")};
+			FType E2T{NestedUsed,		Batch.MakeTypename("E2")};
+			FType E3T{NestedUnused2,	Batch.MakeTypename("E3")};
 			
 			FEnumSchemaId E1D = Batch.DeclareEnum(E1T, EEnumMode::Flat, ELeafWidth::B8, {"C1"}, {1});
 			FEnumSchemaId E2D = Batch.DeclareEnum(E2T, EEnumMode::Flat, ELeafWidth::B8, {"C2"}, {2});
 			FEnumSchemaId E3D = Batch.DeclareEnum(E3T, EEnumMode::Flat, ELeafWidth::B8, {"C3"}, {3});
 
-			FTypeId S1T{NestedUnused1,	Batch.MakeTypename("S1")};
-			FTypeId S2T{NestedUsed,		Batch.MakeTypename("S2")};
-			FTypeId S3T = Batch.MakeParametricType({NestedUnused2,	Batch.MakeTypename("S3")}, {S1T});
-			FTypeId S4T = Batch.MakeParametricType({DoubleNested,	Batch.MakeTypename("S4")}, {S2T, E2T});
-			FTypeId S5T = Batch.MakeParametricType({NestedUnused3,	Batch.MakeTypename("S5")}, {E3T, E1T, S2T});
+			FType S1T{NestedUnused1,	Batch.MakeTypename("S1")};
+			FType S2T{NestedUsed,		Batch.MakeTypename("S2")};
+			FType S3T = Batch.MakeParametricType({NestedUnused2,	Batch.MakeTypename("S3")}, {S1T});
+			FType S4T = Batch.MakeParametricType({DoubleNested,	Batch.MakeTypename("S4")}, {S2T, E2T});
+			FType S5T = Batch.MakeParametricType({NestedUnused3,	Batch.MakeTypename("S5")}, {E3T, E1T, S2T});
 			
 			FStructSchemaId S1D = Batch.DeclareStruct(S1T, {"M1"}, EMemberPresence::AllowSparse);
 			FStructSchemaId S2D = Batch.DeclareStruct(S2T, {"M2"}, EMemberPresence::AllowSparse);
@@ -1638,7 +1638,7 @@ TEST_CASE_NAMED(FPlainPropsReadWriteTest, "System::Core::Serialization::PlainPro
 		[](TConstArrayView<FStructView> Objects, const FTestNameReader& Names)
 		{	
 			FReadBatchId Batch = Objects[0].Schema.Batch;
-			FTypeId S4T = Objects[0].Schema.Resolve().Type;
+			FType S4T = Objects[0].Schema.Resolve().Type;
 
 			FNestedScope DoubleNested = ResolveUntranslatedNestedScope(Batch, S4T.Scope.AsNested());
 			FNestedScope NestedUsed = ResolveUntranslatedNestedScope(Batch, DoubleNested.Outer.AsNested());
@@ -1651,8 +1651,8 @@ TEST_CASE_NAMED(FPlainPropsReadWriteTest, "System::Core::Serialization::PlainPro
 			CHECK(Names[S4.Name.Get().Id] == "S4");
 			CHECK(S4.NumParameters == 2);
 			
-			FTypeId S2T = S4.Parameters[0];
-			FTypeId E2T = S4.Parameters[1];
+			FType S2T = S4.Parameters[0];
+			FType E2T = S4.Parameters[1];
 			CHECK(S2T.Scope == DoubleNested.Outer);
 			CHECK(E2T.Scope == DoubleNested.Outer);
 			CHECK(Names[S2T.Name.AsConcrete().Id] == "S2");
diff --git a/Engine/Plugins/Experimental/PlainProps/Source/Tests/PlainPropsUeCoreTests.cpp b/Engine/Plugins/Experimental/PlainProps/Source/Tests/PlainPropsUeCoreTests.cpp
--- a/Engine/Plugins/Experimental/PlainProps/Source/Tests/PlainPropsUeCoreTests.cpp
+++ b/Engine/Plugins/Experimental/PlainProps/Source/Tests/PlainPropsUeCoreTests.cpp
@@ -43,8 +43,8 @@ struct FRuntimeIds
 	static FMemberId			IndexMember(FAnsiStringView Name)	{ return GNames.NameMember(FName(Name)); }
 	static FConcreteTypenameId	IndexTypename(FAnsiStringView Name)	{ return GNames.NameType(FName(Name)); }
 	static FScopeId				IndexScope(FAnsiStringView Name)	{ return GNames.MakeScope(FName(Name)); }
-	static FEnumSchemaId		IndexEnum(FTypeId Type)				{ return GNames.IndexEnum(Type); }
-	static FStructSchemaId		IndexStruct(FTypeId Type)			{ return GNames.IndexStruct(Type); }
+	static FEnumSchemaId		IndexEnum(FType Type)				{ return GNames.IndexEnum(Type); }
+	static FStructSchemaId		IndexStruct(FType Type)			{ return GNames.IndexStruct(Type); }
 	static FIdIndexerBase&		GetIndexer()						{ return GNames; }
 };
 
@@ -93,7 +93,7 @@ struct FNameDeclaration
 	FStructSchemaId		Id;
 	FMemberId			Idx;
 
-	FNameDeclaration(FTypeId Type = IndexStructName<FRuntimeIds, ETypename::Decl, TTypename<FName>>())
+	FNameDeclaration(FType Type = IndexStructName<FRuntimeIds, ETypename::Decl, TTypename<FName>>())
 	: Id(FRuntimeIds::IndexStruct(Type))
 	, Idx(FRuntimeIds::IndexMember("Idx"))
 	{
@@ -272,8 +272,8 @@ public:
 	virtual uint32				NumEnums() const override final						{ return Schemas.NumSchemas - Schemas.NumStructSchemas; }
 	virtual uint32				NumStructs() const override final					{ return Schemas.NumStructSchemas; }
 
-	virtual FTypeId				Resolve(FEnumSchemaId Id) const override final		{ return ResolveEnumSchema(Schemas, Id).Type; }
-	virtual FTypeId				Resolve(FStructSchemaId Id) const override final	{ return ResolveStructSchema(Schemas, Id).Type; }
+	virtual FType				Resolve(FEnumSchemaId Id) const override final		{ return ResolveEnumSchema(Schemas, Id).Type; }
+	virtual FType				Resolve(FStructSchemaId Id) const override final	{ return ResolveStructSchema(Schemas, Id).Type; }
 
 protected:
 	const FSchemaBatch& Schemas;
@@ -324,8 +324,8 @@ public:
 
 	virtual FNestedScope		Resolve(FNestedScopeId Id) const override final		{ return Ids.Resolve(Binding.Remap(Id)); }
 	virtual FParametricTypeView	Resolve(FParametricTypeId Id) const override final	{ return Ids.Resolve(Binding.Remap(Id)); }
-	virtual FTypeId				Resolve(FEnumSchemaId Id) const override final		{ return ResolveEnumSchema(Schemas, Id).Type; }
-	virtual FTypeId				Resolve(FStructSchemaId Id) const override final	{ return ResolveStructSchema(Schemas, Id).Type; }
+	virtual FType				Resolve(FEnumSchemaId Id) const override final		{ return ResolveEnumSchema(Schemas, Id).Type; }
+	virtual FType				Resolve(FStructSchemaId Id) const override final	{ return ResolveStructSchema(Schemas, Id).Type; }
 
 	using FIds::AppendString;
 	virtual void AppendString(FUtf8Builder& Out, FNameId Name) const override final			{ Ids.AppendString(Out, Name); }
diff --git a/Engine/Plugins/Experimental/PlainPropsUObject/Source/Private/PlainPropsUObjectBindings.cpp b/Engine/Plugins/Experimental/PlainPropsUObject/Source/Private/PlainPropsUObjectBindings.cpp
--- a/Engine/Plugins/Experimental/PlainPropsUObject/Source/Private/PlainPropsUObjectBindings.cpp
+++ b/Engine/Plugins/Experimental/PlainPropsUObject/Source/Private/PlainPropsUObjectBindings.cpp
@@ -434,7 +434,7 @@ static const UStruct* SkipEmptyBases(const UStruct* In)
 	return In;
 }
 
-static FTypeId IndexType(const UField* Field)
+static FType IndexType(const UField* Field)
 {
 	check(Field);
 	FTypenameId Name = GUE.Names.MakeTypename(Field->GetFName());
@@ -499,7 +499,7 @@ static void DeclareSuperMembers(FMemberArray& Out, const UStruct* Struct)
 	}
 }
 
-static void DeclareStruct(const UStruct* Struct, FTypeId Type, FStructSchemaId Id)
+static void DeclareStruct(const UStruct* Struct, FType Type, FStructSchemaId Id)
 {
 	FOptionalStructSchemaId  Super = IndexSuper(Struct);
 	EMemberPresence Occupancy = GetOccupancy(Struct);
@@ -518,7 +518,7 @@ static void DeclareStruct(const UStruct* Struct, FTypeId Type, FStructSchemaId I
 
 static FStructSchemaId DeclareStruct(const UStruct* Struct)
 {
-	FTypeId Type = IndexType(Struct);
+	FType Type = IndexType(Struct);
 	FStructSchemaId Id = GUE.Names.IndexStruct(Type);
 	DeclareStruct(Struct, Type, Id);
 	return Id;
@@ -810,10 +810,10 @@ static FStructSchemaId BindSparseDelegate(FStructSchemaId Owner, FMulticastSpars
 	// Todo: Ownership / memory leak
 	ICustomBinding* Leak = new FMulticastSparseDelegateBinding(CastChecked<USparseDelegateFunction>(Property->SignatureFunction));
 
-	FTypeId MulticastSparseDelegate = GUE.Names.Resolve(GUE.Structs.MulticastSparseDelegate);
-	FTypeId OwnerParam = GUE.Names.Resolve(Owner);
-	FTypeId PropertyParam = { GUE.Scopes.CoreUObject, FTypenameId(GUE.Names.NameType(Property->GetFName())) };
-	FTypeId UniqueBindName = GUE.Names.MakeParametricType(MulticastSparseDelegate, {OwnerParam, PropertyParam});
+	FType MulticastSparseDelegate = GUE.Names.Resolve(GUE.Structs.MulticastSparseDelegate);
+	FType OwnerParam = GUE.Names.Resolve(Owner);
+	FType PropertyParam = { GUE.Scopes.CoreUObject, FTypenameId(GUE.Names.NameType(Property->GetFName())) };
+	FType UniqueBindName = GUE.Names.MakeParametricType(MulticastSparseDelegate, {OwnerParam, PropertyParam});
 	FStructSchemaId BindId = GUE.Names.IndexStruct(UniqueBindName);
 
 	GUE.Customs.BindStruct(BindId, *Leak, GUE.Types.Get(GUE.Structs.MulticastDelegate));
@@ -918,10 +918,10 @@ struct FInterfaceBinding final : ICustomBinding
 
 class FInterfaceBindings
 {
-	const FTypeId ScriptInterface;
+	const FType ScriptInterface;
 	const FStructDeclaration& Declaration;
 
-	TMap<FTypeId, FStructSchemaId> BoundClasses;
+	TMap<FType, FStructSchemaId> BoundClasses;
 
 public:
 	FInterfaceBindings()
@@ -931,13 +931,13 @@ public:
 
 	FStructSchemaId Bind(FInterfaceProperty* Property)
 	{
-		FTypeId Class = IndexType(Property->InterfaceClass);
+		FType Class = IndexType(Property->InterfaceClass);
 		if (const FStructSchemaId* BindId = BoundClasses.Find(Class))
 		{
 			return *BindId;
 		}
 		
-		FTypeId UniqueBindName = GUE.Names.MakeParametricType(ScriptInterface, {Class});
+		FType UniqueBindName = GUE.Names.MakeParametricType(ScriptInterface, {Class});
 		FStructSchemaId BindId = GUE.Names.IndexStruct(UniqueBindName);
 		BoundClasses.Emplace(Class, BindId);
 
@@ -1014,7 +1014,7 @@ static bool DenyMaxValue(FName Enum)
 
 static FEnumSchemaId DeclareEnum(UEnum* Enum)
 {
-	FTypeId Type = IndexType(Enum);
+	FType Type = IndexType(Enum);
 	FEnumSchemaId Id = GUE.Names.IndexEnum(Type);
 	EEnumMode Mode = Enum->HasAnyEnumFlags(EEnumFlags::Flags) ? EEnumMode::Flag : EEnumMode::Flat;
 
@@ -2245,10 +2245,10 @@ public:
 		}
 
 		// Index custom delta binding struct name
-		FDualTypeId Param = Elem.IndexParameterName(GUE.Names);
-		FTypeId BindType = FTypeId{ GUE.Scopes.Core, FTypenameId(GUE.Names.MakeParametricTypeId(GUE.Typenames.Set, MakeArrayView(&Param.BindType, 1))) };
-		FTypeId DeclType = Param.IsLowered()
-						 ? FTypeId{ GUE.Scopes.Core, FTypenameId(GUE.Names.MakeParametricTypeId(GUE.Typenames.Set, MakeArrayView(&Param.DeclType, 1))) }
+		FDualType Param = Elem.IndexParameterName(GUE.Names);
+		FType BindType = FType{ GUE.Scopes.Core, FTypenameId(GUE.Names.MakeParametricTypeId(GUE.Typenames.Set, MakeArrayView(&Param.BindType, 1))) };
+		FType DeclType = Param.IsLowered()
+						 ? FType{ GUE.Scopes.Core, FTypenameId(GUE.Names.MakeParametricTypeId(GUE.Typenames.Set, MakeArrayView(&Param.DeclType, 1))) }
 						 : BindType;
 		FStructSchemaId BindId = GUE.Names.IndexStruct(BindType);
 		FStructSchemaId DeclId = Param.IsLowered() ? GUE.Names.IndexStruct(DeclType) : BindId;
@@ -2584,8 +2584,8 @@ class FMapBindings
 
 		// Index custom delta binding struct name
 		FParametricTypeId PairTypename = GUE.Names.Resolve(static_cast<FStructSchemaId>(Members.Pair.InnermostSchema.Get())).Name.AsParametric();
-		TConstArrayView<FTypeId> Params = GUE.Names.Resolve(PairTypename).GetParameters();
-		FTypeId Type = { GUE.Scopes.Core, FTypenameId(GUE.Names.MakeParametricTypeId(GUE.Typenames.Map, Params)) };
+		TConstArrayView<FType> Params = GUE.Names.Resolve(PairTypename).GetParameters();
+		FType Type = { GUE.Scopes.Core, FTypenameId(GUE.Names.MakeParametricTypeId(GUE.Typenames.Map, Params)) };
 		FStructSchemaId Id = GUE.Names.IndexStruct(Type);
 		FMemberId MemberIds[] = { GUE.Members.Assign, GUE.Members.Remove, GUE.Members.Insert };
 
@@ -2681,12 +2681,12 @@ class FPairBindings
 		}
 
 		// Index names, can be optimized by checking if KeyParam / BindParam IsLowered()
-		FDualTypeId KeyParam = Pair.KV[0].IndexParameterName(GUE.Names);
-		FDualTypeId ValueParam = Pair.KV[1].IndexParameterName(GUE.Names);
-		FTypeId BindParams[2] = { KeyParam.BindType, ValueParam.BindType };
-		FTypeId DeclParams[2] = { KeyParam.DeclType, ValueParam.DeclType };
-		FTypeId BindType = { GUE.Scopes.Core, FTypenameId(GUE.Names.MakeParametricTypeId(GUE.Typenames.Pair, BindParams)) };
-		FTypeId DeclType = { GUE.Scopes.Core, FTypenameId(GUE.Names.MakeParametricTypeId(GUE.Typenames.Pair, DeclParams)) };
+		FDualType KeyParam = Pair.KV[0].IndexParameterName(GUE.Names);
+		FDualType ValueParam = Pair.KV[1].IndexParameterName(GUE.Names);
+		FType BindParams[2] = { KeyParam.BindType, ValueParam.BindType };
+		FType DeclParams[2] = { KeyParam.DeclType, ValueParam.DeclType };
+		FType BindType = { GUE.Scopes.Core, FTypenameId(GUE.Names.MakeParametricTypeId(GUE.Typenames.Pair, BindParams)) };
+		FType DeclType = { GUE.Scopes.Core, FTypenameId(GUE.Names.MakeParametricTypeId(GUE.Typenames.Pair, DeclParams)) };
 		FStructSchemaId BindId = GUE.Names.IndexStruct(BindType);
 		FStructSchemaId DeclId = GUE.Names.IndexStruct(DeclType);
 		FMemberId Members[2] = { GUE.Members.Key, GUE.Members.Value };
@@ -3003,7 +3003,7 @@ private:
 	{
 		if (!OwnerScope)
 		{
-			FTypeId OwnerType = GUE.Names.Resolve(Owner);
+			FType OwnerType = GUE.Names.Resolve(Owner);
 			OwnerScope = GUE.Names.NestFlatScope(OwnerType.Scope, {OwnerType.Name.AsConcrete().Id});
 		}
 
@@ -3046,7 +3046,7 @@ private:
 
 	inline FMemberBinding BindAsStruct(FProperty* Property, UStruct* Struct)
 	{
-		FTypeId Type = IndexType(SkipEmptyBases(Struct));
+		FType Type = IndexType(SkipEmptyBases(Struct));
 		return BindAsStruct(Property, GUE.Names.IndexStruct(Type));
 	}
 
@@ -3212,7 +3212,7 @@ private:
 		return FMemberBinding(Property->GetOffset_ForInternal());
 	}
 
-	FTypeId MakeStaticArrayTypename(FName PropertyName)
+	FType MakeStaticArrayTypename(FName PropertyName)
 	{
 		return { GetOwnerScope(), GUE.Names.MakeTypename(PropertyName) };
 	}
@@ -3249,7 +3249,7 @@ private:
 		else
 		{
 			// Create struct binding to allow delta serialization
-			FTypeId StaticArrayType = MakeStaticArrayTypename(Property->GetFName());
+			FType StaticArrayType = MakeStaticArrayTypename(Property->GetFName());
 			FStructSchemaId StaticArrayId = GUE.Names.IndexStruct(StaticArrayType);
 
 			// Todo: Ownership
@@ -3350,7 +3350,7 @@ static void BindInitialTypes()
 		UField* Field = *It;
 		if (UStruct* Struct = Cast<UStruct>(Field))
 		{
-			FTypeId Type = IndexType(Struct);
+			FType Type = IndexType(Struct);
 			FStructSchemaId Id = GUE.Names.IndexStruct(Type);
 			if (!Algo::Find(SkipStructs, Id))
 			{
@@ -3568,7 +3568,7 @@ struct TCustomPropertyBindings
 struct FMemoryBatch
 {
 	TArray64<uint8>				Data;
-	TArray<FStructSchemaId>		RuntimeIds; // To avoid reindexing schema FTypeId
+	TArray<FStructSchemaId>		RuntimeIds; // To avoid reindexing schema FType
 	FMemoryPropertyBatch		Properties; // Contains FTexts referenced from Data
 };
 
@@ -4179,6 +4179,6 @@ extern FString _print(const PlainProps::FOptionalMemberId& Id)		{ return PlainPr
 extern FString _print(const PlainProps::FScopeId& Id)				{ return PlainProps::UE::GUE.Debug.Print(Id); }
 extern FString _print(const PlainProps::FTypenameId& Id)			{ return PlainProps::UE::GUE.Debug.Print(Id); }
 extern FString _print(const PlainProps::FConcreteTypenameId& Id)	{ return PlainProps::UE::GUE.Debug.Print(Id); }
-extern FString _print(const PlainProps::FTypeId& Id)				{ return PlainProps::UE::GUE.Debug.Print(Id); }
+extern FString _print(const PlainProps::FType& Id)				{ return PlainProps::UE::GUE.Debug.Print(Id); }
 extern FString _print(const PlainProps::FEnumSchemaId& Id)			{ return PlainProps::UE::GUE.Debug.Print(Id); }
 extern FString _print(const PlainProps::FStructSchemaId& Id)		{ return PlainProps::UE::GUE.Debug.Print(Id); }
diff --git a/Engine/Plugins/Experimental/PlainPropsUObject/Source/Public/PlainPropsUObjectRuntime.h b/Engine/Plugins/Experimental/PlainPropsUObject/Source/Public/PlainPropsUObjectRuntime.h
--- a/Engine/Plugins/Experimental/PlainPropsUObject/Source/Public/PlainPropsUObjectRuntime.h
+++ b/Engine/Plugins/Experimental/PlainPropsUObject/Source/Public/PlainPropsUObjectRuntime.h
@@ -156,8 +156,8 @@ struct FRuntimeIds
 	static FMemberId			IndexMember(FAnsiStringView Name)		{ return GUE.Names.NameMember(FName(Name)); }
 	static FConcreteTypenameId	IndexTypename(FAnsiStringView Name)		{ return GUE.Names.NameType(FName(Name)); }
 	static FFlatScopeId			IndexScope(FAnsiStringView Name)		{ return GUE.Names.NameScope(FName(Name)); }
-	static FEnumSchemaId		IndexEnum(FTypeId Type)					{ return GUE.Names.IndexEnum(Type); }
-	static FStructSchemaId		IndexStruct(FTypeId Type)				{ return GUE.Names.IndexStruct(Type); }
+	static FEnumSchemaId		IndexEnum(FType Type)					{ return GUE.Names.IndexEnum(Type); }
+	static FStructSchemaId		IndexStruct(FType Type)					{ return GUE.Names.IndexStruct(Type); }
 	static FIdIndexerBase&		GetIndexer()							{ return GUE.Names; }
 };
 
