PlainProps: Add text roundtripping to commandlet.
Add free Print functions instead of relying on FUtf8StringBuilder formatting.
Add proper unicode escaping (json standard) and some basic unicode unit tests.
Disable one unit test item related to empty ranges of dynamic types - to revisit later.
Add Diff API for comparing struct views.

#rnx

diff --git a/Engine/Plugins/Experimental/PlainProps/Source/Private/PlainPropsDiff.cpp b/Engine/Plugins/Experimental/PlainProps/Source/Private/PlainPropsDiff.cpp
--- a/Engine/Plugins/Experimental/PlainProps/Source/Private/PlainPropsDiff.cpp
+++ b/Engine/Plugins/Experimental/PlainProps/Source/Private/PlainPropsDiff.cpp
@@ -2,8 +2,11 @@
 
 #include "PlainPropsDiff.h"
 #include "PlainPropsBind.h"
+#include "PlainPropsBuild.h"
 #include "PlainPropsInternalBind.h"
 #include "PlainPropsInternalDiff.h"
+#include "PlainPropsInternalFormat.h"
+#include "PlainPropsInternalRead.h"
 
 namespace PlainProps
 {
@@ -142,4 +145,365 @@ bool DiffStructs(FStructId BindId, const void* A, const void* B, FDiffContext& C
 	return DiffMembers(BindId, A, B, Ctx);
 }
 
-} // namespace PlainProps
\ No newline at end of file
+////////////////////////////////////////////////////////////////////////////////
+
+static uint64 CalculateSize(const FSchemaBatch& In)
+{
+	uint32 NumParameters = 0;
+	for (FParametricType ParametricType : In.GetParametricTypes())
+	{
+		NumParameters += ParametricType.Parameters.NumParameters;
+	}
+	const uint8* End = reinterpret_cast<const uint8*>(In.GetFirstParameter() + NumParameters);
+	const uint64 Size = End - reinterpret_cast<const uint8*>(&In);
+	return Size;
+}
+
+static bool operator==(const FSchemaBatch& A, const FSchemaBatch& B)
+{
+	const uint64 SizeA = CalculateSize(A);
+	const uint64 SizeB = CalculateSize(B);
+	return SizeA == SizeB && FMemory::Memcmp(&A, &B, SizeA) == 0;
+}
+
+static bool operator==(const FStructSchema& A, const FStructSchema& B)
+{
+	const uint64 SizeA = CalculateSize(A);
+	const uint64 SizeB = CalculateSize(B);
+	return SizeA == SizeB && FMemory::Memcmp(&A, &B, SizeA) == 0;
+}
+
+static bool operator==(const FEnumSchema& A, const FEnumSchema& B)
+{
+	const uint64 SizeA = CalculateSize(A);
+	const uint64 SizeB = CalculateSize(B);
+	return SizeA == SizeB && FMemory::Memcmp(&A, &B, SizeA) == 0;
+}
+
+static bool operator!=(const FSchemaBatch& A, const FSchemaBatch& B)	{ return !(A == B); }
+static bool operator!=(const FStructSchema& A, const FStructSchema& B)	{ return !(A == B); }
+static bool operator!=(const FEnumSchema& A, const FEnumSchema& B)		{ return !(A == B); }
+
+////////////////////////////////////////////////////////////////////////////////
+
+static bool DiffStructSchema(const FStructSchema& A, const FStructSchema& B)
+{
+	const bool Diff = A != B;
+#if DO_CHECK
+	check(A.Type == B.Type);
+	check(A.NumMembers == B.NumMembers);
+	check(A.NumRangeTypes == B.NumRangeTypes);
+	check(A.NumInnerSchemas == B.NumInnerSchemas);
+	check(A.Inheritance == B.Inheritance);
+	check(A.IsDense == B.IsDense);
+	check(EqualItems(A.GetMemberTypes(), B.GetMemberTypes()));
+	check(EqualItems(A.GetRangeTypes(), B.GetRangeTypes()));
+	check(EqualItems(A.GetMemberNames(), B.GetMemberNames()));
+	check(EqualItems(
+		MakeArrayView(A.GetInnerSchemas(), A.NumInnerSchemas),
+		MakeArrayView(B.GetInnerSchemas(), B.NumInnerSchemas)));
+	check(A.GetSuperSchema() == B.GetSuperSchema());
+	check(A == B);
+#endif
+	return Diff;
+}
+
+static bool DiffEnumSchema(const FEnumSchema& A, const FEnumSchema& B)
+{
+	const bool Diff = A != B;
+#if DO_CHECK
+	check(A.Type == B.Type);
+	check(A.FlagMode == B.FlagMode);
+	check(A.ExplicitConstants == B.ExplicitConstants);
+	check(A.Width == B.Width);
+	check(A.Num == B.Num);
+	check(EqualItems(MakeArrayView(A.Footer, A.Num), MakeArrayView(B.Footer, B.Num)));
+	if (A.Width == B.Width)
+	{
+		switch (A.Width)
+		{
+		case ELeafWidth::B8:  check(EqualItems(GetConstants<uint8 >(A), GetConstants<uint8 >(B))); break;
+		case ELeafWidth::B16: check(EqualItems(GetConstants<uint16>(A), GetConstants<uint16>(B))); break;
+		case ELeafWidth::B32: check(EqualItems(GetConstants<uint32>(A), GetConstants<uint32>(B))); break;
+		case ELeafWidth::B64: check(EqualItems(GetConstants<uint64>(A), GetConstants<uint64>(B))); break;
+		}
+	}
+	check(A == B);
+#endif
+	return Diff;
+}
+
+static bool DiffSchemas(const FSchemaBatch& A, const FSchemaBatch& B)
+{
+	const bool Diff = A != B;
+#if DO_CHECK
+	check(A.NumNestedScopes == B.NumNestedScopes);
+	check(A.NestedScopesOffset == B.NestedScopesOffset);
+	check(A.NumParametricTypes == B.NumParametricTypes);
+	check(A.NumSchemas == B.NumSchemas);
+	check(A.NumStructSchemas == B.NumStructSchemas);
+	check(EqualItems(A.GetSchemaOffsets(), B.GetSchemaOffsets()));
+	check(EqualItems(A.GetNestedScopes(), B.GetNestedScopes()));
+	check(EqualItems(A.GetParametricTypes(), B.GetParametricTypes()));
+
+	{
+		TSchemaRange<const FStructSchema> AA = GetStructSchemas(A);
+		TSchemaRange<const FStructSchema> BB = GetStructSchemas(B);
+		TSchemaIterator<const FStructSchema> ItA = AA.begin();
+		TSchemaIterator<const FStructSchema> ItB = BB.begin();
+		for (; ItA != AA.end() && ItB != BB.end(); ++ItA, ++ItB)
+		{
+			check(!DiffStructSchema(*ItA, *ItB));
+		}
+	}
+
+	{
+		TSchemaRange<const FEnumSchema> AA = GetEnumSchemas(A);
+		TSchemaRange<const FEnumSchema> BB = GetEnumSchemas(B);
+		TSchemaIterator<const FEnumSchema> ItA = AA.begin();
+		TSchemaIterator<const FEnumSchema> ItB = BB.begin();
+		for (; ItA != AA.end() && ItB != BB.end(); ++ItA, ++ItB)
+		{
+			check(!DiffEnumSchema(*ItA, *ItB));
+		}
+	}
+
+	check(!Diff);
+#endif
+	return Diff;
+}
+
+bool DiffSchemas(FReadBatchId A, FReadBatchId B)
+{
+	return DiffSchemas(GetReadSchemas(A), GetReadSchemas(B));
+}
+
+////////////////////////////////////////////////////////////////////////////////
+
+static bool DiffLeaf(FLeafView A, FLeafView B)
+{
+	if (A.Leaf != B.Leaf || A.Enum != B.Enum)
+	{
+		return true;
+	}
+
+	if (A.Leaf.Type == ELeafType::Bool)
+	{
+		return A.Value.bValue != B.Value.bValue;
+	}
+
+	switch (A.Leaf.Width)
+	{
+	case ELeafWidth::B8:  return !!FMemory::Memcmp(A.Value.Ptr, B.Value.Ptr, 1);
+	case ELeafWidth::B16: return !!FMemory::Memcmp(A.Value.Ptr, B.Value.Ptr, 2);
+	case ELeafWidth::B32: return !!FMemory::Memcmp(A.Value.Ptr, B.Value.Ptr, 4);
+	case ELeafWidth::B64: return !!FMemory::Memcmp(A.Value.Ptr, B.Value.Ptr, 8);
+	}
+
+	unimplemented();
+	return false;
+}
+
+static bool DiffLeaves(FUnpackedLeafType Leaf, FLeafRangeView A, FLeafRangeView B, uint64& OutIdx)
+{
+	OutIdx = ~0u;
+	switch (Leaf.Type)
+	{
+	case ELeafType::Bool:
+		check(Leaf.Width == ELeafWidth::B8);
+		return DiffItems(A.AsBools(), B.AsBools(), OutIdx);
+	case ELeafType::IntS:
+		switch (Leaf.Width)
+		{
+		case ELeafWidth::B8:  return DiffItems(A.AsS8s(),  B.AsS8s(),  OutIdx);
+		case ELeafWidth::B16: return DiffItems(A.AsS16s(), B.AsS16s(), OutIdx);
+		case ELeafWidth::B32: return DiffItems(A.AsS32s(), B.AsS32s(), OutIdx);
+		case ELeafWidth::B64: return DiffItems(A.AsS64s(), B.AsS64s(), OutIdx);
+		}
+	case ELeafType::IntU:
+		switch (Leaf.Width)
+		{
+		case ELeafWidth::B8:  return DiffItems(A.AsU8s(),  B.AsU8s(),  OutIdx);
+		case ELeafWidth::B16: return DiffItems(A.AsU16s(), B.AsU16s(), OutIdx);
+		case ELeafWidth::B32: return DiffItems(A.AsU32s(), B.AsU32s(), OutIdx);
+		case ELeafWidth::B64: return DiffItems(A.AsU64s(), B.AsU64s(), OutIdx);
+		}
+	case ELeafType::Float:
+		if (Leaf.Width == ELeafWidth::B32)
+		{
+			return DiffItems(A.AsFloats(), B.AsFloats(), OutIdx);
+		}
+		check(Leaf.Width == ELeafWidth::B64);
+		return DiffItems(A.AsDoubles(), B.AsDoubles(), OutIdx);
+	case ELeafType::Hex:
+		// PP-TEXT: Implement DiffItems(Hex)
+		check(Leaf.Type != ELeafType::Hex);
+		break;
+	case ELeafType::Enum:
+		switch (Leaf.Width)
+		{
+		case ELeafWidth::B8:  return DiffItems(A.AsUnderlyingValues<uint8>(), B.AsUnderlyingValues<uint8>(), OutIdx);
+		case ELeafWidth::B16: return DiffItems(A.AsUnderlyingValues<uint16>(), B.AsUnderlyingValues<uint16>(), OutIdx);
+		case ELeafWidth::B32: return DiffItems(A.AsUnderlyingValues<uint32>(), B.AsUnderlyingValues<uint32>(), OutIdx);
+		case ELeafWidth::B64: return DiffItems(A.AsUnderlyingValues<uint64>(), B.AsUnderlyingValues<uint64>(), OutIdx);
+		}
+	case ELeafType::Unicode:
+		switch (Leaf.Width)
+		{
+		case ELeafWidth::B8:  return DiffItems(A.AsUtf8(),  B.AsUtf8(),  OutIdx);
+		case ELeafWidth::B16: return DiffItems(A.AsUtf16(), B.AsUtf16(), OutIdx);
+		case ELeafWidth::B32: return DiffItems(A.AsUtf32(), B.AsUtf32(), OutIdx);
+		case ELeafWidth::B64: check(false);
+		}
+	}
+	unimplemented();
+	return false;
+}
+
+static bool DiffMembers(FStructView A, FStructView B, FReadDiffPath& Out);
+static bool DiffRange(ERangeSizeType NumType, FRangeView A, FRangeView B, FReadDiffPath& Out);
+
+static bool DiffStructs(FStructType Struct, FStructRangeView A, FStructRangeView B, FReadDiffPath& Out)
+{
+	uint64 NumA = A.Num();
+	uint64 NumB = B.Num();
+
+	uint64 DiffIdx = 0;
+	bool bDiff = false;
+	FStructRangeIterator ItB = B.begin();
+	for (FStructView ItA : A)
+	{
+		if (DiffIdx >= NumB)
+		{
+			break;
+		}
+		bDiff = DiffMembers(ItA, *ItB, Out);
+		if (bDiff)
+		{
+			break;
+		}
+		++ItB;
+		++DiffIdx;
+	}
+	if (bDiff || NumA != NumB)
+	{
+		Out.Emplace(FMemberType(Struct), NoId, NoId, DiffIdx);
+		return true;
+	}
+	return false;
+}
+
+static bool DiffRanges(ERangeSizeType NumType, FNestedRangeView A, FNestedRangeView B, FReadDiffPath& Out)
+{
+	uint64 NumA = A.Num();
+	uint64 NumB = B.Num();
+	
+	uint64 DiffIdx = 0;
+	bool bDiff = false;
+	FNestedRangeIterator ItB = B.begin();
+	for (FRangeView ItA : A)
+	{
+		if (DiffIdx >= NumB)
+		{
+			break;
+		}
+		bDiff = DiffRange(NumType, ItA, *ItB, Out);
+		if (bDiff)
+		{
+			break;
+		}
+		++ItB;
+		++DiffIdx;
+	}
+	if (bDiff || NumA != NumB)
+	{
+		Out.Emplace(FMemberType(NumType), NoId, NoId, DiffIdx);
+		return true;
+	}
+	return false;
+}
+
+static bool DiffLeaves(FUnpackedLeafType Leaf, FLeafRangeView A, FLeafRangeView B, FReadDiffPath& Out)
+{
+	uint64 DiffIdx;
+	if (DiffLeaves(Leaf, A, B, DiffIdx))
+	{
+		Out.Emplace(Leaf.Pack(), NoId, NoId, DiffIdx);
+		return true;
+	}
+	return false;
+}
+
+static bool DiffRange(ERangeSizeType NumType, FRangeView A, FRangeView B, FReadDiffPath& Out)
+{
+	bool bDiff = A.GetItemType() != B.GetItemType();
+	if (!bDiff)
+	{
+		switch (A.GetItemType().GetKind())
+		{
+		case EMemberKind::Leaf:
+			bDiff = DiffLeaves(A.GetItemType().AsLeaf(), A.AsLeaves(), B.AsLeaves(), Out);
+			break;
+		case EMemberKind::Struct:
+			bDiff = DiffStructs(A.GetItemType().AsStruct(), A.AsStructs(), B.AsStructs(), Out);
+			break;
+		case EMemberKind::Range:
+			bDiff = DiffRanges(A.GetItemType().AsRange().MaxSize, A.AsRanges(), B.AsRanges(), Out);
+			break;
+		}
+	}
+	return bDiff;
+}
+
+static bool DiffMembers(FStructView A, FStructView B, FReadDiffPath& Out)
+{
+	FMemberReader ItA(A);
+	FMemberReader ItB(B);
+
+	for (; ItA.HasMore() && ItB.HasMore(); )
+	{
+		FOptionalMemberId NameA = ItA.PeekName();
+		FOptionalMemberId NameB = ItB.PeekName();
+		FMemberType TypeA = ItA.PeekType();
+		FMemberType TypeB = ItB.PeekType();
+		bool bDiff = TypeA != TypeB || NameA != NameB;
+		if (!bDiff)
+		{
+			switch (TypeA.GetKind())
+			{
+			case EMemberKind::Leaf:
+				bDiff = DiffLeaf(ItA.GrabLeaf(), ItB.GrabLeaf());
+				break;
+			case EMemberKind::Struct:
+				bDiff = DiffMembers(ItA.GrabStruct(), ItB.GrabStruct(), Out);
+				break;
+			case EMemberKind::Range:
+				bDiff = DiffRange(TypeA.AsRange().MaxSize, ItA.GrabRange(), ItB.GrabRange(), Out);
+				break;
+			}
+		}
+		if (bDiff)
+		{
+			Out.Emplace(TypeA, A.Schema.Id, NameA);
+			return true;
+		}
+	}
+	if (ItA.HasMore())
+	{
+		Out.Emplace(ItA.PeekType(), A.Schema.Id, ItA.PeekName());
+		return true;
+	}
+	if (ItB.HasMore())
+	{
+		Out.Emplace(ItB.PeekType(), A.Schema.Id, ItB.PeekName());
+		return true;
+	}
+	return false;
+}
+
+bool DiffStruct(FStructView A, FStructView B, FReadDiffPath& Out)
+{
+	return DiffMembers(A, B, Out);
+}
+
+} // namespace PlainProps
diff --git a/Engine/Plugins/Experimental/PlainProps/Source/Private/PlainPropsInternalDiff.h b/Engine/Plugins/Experimental/PlainProps/Source/Private/PlainPropsInternalDiff.h
--- a/Engine/Plugins/Experimental/PlainProps/Source/Private/PlainPropsInternalDiff.h
+++ b/Engine/Plugins/Experimental/PlainProps/Source/Private/PlainPropsInternalDiff.h
@@ -14,6 +14,29 @@ inline const uint8* At(const void* Ptr, SIZE_T Offset)
 	return static_cast<const uint8*>(Ptr) + Offset;
 }
 
+//  Works with FBoolRangeView, which lack GetData(), TRangeView, TArrayView, std::initializer_list and T[]
+template<class RangeTypeA, class RangeTypeB>
+bool DiffItems(RangeTypeA&& A, RangeTypeB&& B, uint64& OutIdx)
+{
+	uint64 NumA = IntCastChecked<uint64>(GetNum(A));
+	uint64 NumB = IntCastChecked<uint64>(GetNum(B));
+	OutIdx = 0;
+	auto BIt = std::begin(B);
+	for (auto bA : A)
+	{
+		if (OutIdx >= NumB)
+		{
+			break;
+		}
+		if (bA != *BIt++)
+		{
+			return true;
+		}
+		++OutIdx;
+	}
+	return NumA != NumB;
+}
+
 inline int64 DiffLeaf(const uint8* A, const uint8* B, FUnpackedLeafBindType Leaf)
 {
 	bool bBitfield = Leaf.Type == ELeafBindType::BitfieldBool;
@@ -152,4 +175,4 @@ bool DiffItemRange(const void* RangeA, const void* RangeB, const IItemRangeBindi
 	return false;
 }
 
-} // namespace PlainProps
\ No newline at end of file
+} // namespace PlainProps
diff --git a/Engine/Plugins/Experimental/PlainProps/Source/Private/PlainPropsInternalParse.h b/Engine/Plugins/Experimental/PlainProps/Source/Private/PlainPropsInternalParse.h
--- a/Engine/Plugins/Experimental/PlainProps/Source/Private/PlainPropsInternalParse.h
+++ b/Engine/Plugins/Experimental/PlainProps/Source/Private/PlainPropsInternalParse.h
@@ -97,6 +97,16 @@ struct FParsedSchemas
 class FTextIndexer : public TIdIndexer<FUtf8StringView>
 {
 public:
+	FOptionalMemberId GetMemberId(FUtf8StringView Name)
+	{
+		FSetElementId Id = Names.FindId(Name);
+		if (Id.IsValidId())
+		{
+			return FMemberId{ IntCastChecked<uint32>(Id.AsInteger()) };
+		}
+		return NoId;
+	}
+
 	FOptionalStructId GetStructId(FType Type) const
 	{
 		FSetElementId Id = Structs.FindId(Type);
@@ -145,6 +155,7 @@ public:
 	void Parse(TArray64<uint8>& Out);
 
 private:
+	void Tokenize();
 	void ParseAll();
 	void ParseStructSchemas(TConstArrayView64<FToken> TokensView);
 	void ParseEnumSchemas(TConstArrayView64<FToken> TokensView);
@@ -157,6 +168,11 @@ private:
 	FDeclarations									Types;
 	TArray<TTuple<FStructId, FBuiltStruct*>>		Objects;
 	FParsedSchemas									Schemas;
+	TArray64<FToken>								Tokens;
+	int64											EnumsIdx   = -1;
+	int64											StructsIdx = -1;
+	int64											ObjectsIdx = -1;
+
 };
 
 } // namespace PlainProps
\ No newline at end of file
diff --git a/Engine/Plugins/Experimental/PlainProps/Source/Private/PlainPropsInternalPrint.h b/Engine/Plugins/Experimental/PlainProps/Source/Private/PlainPropsInternalPrint.h
--- a/Engine/Plugins/Experimental/PlainProps/Source/Private/PlainPropsInternalPrint.h
+++ b/Engine/Plugins/Experimental/PlainProps/Source/Private/PlainPropsInternalPrint.h
@@ -3,6 +3,7 @@
 #pragma once
 
 #include "Containers/ArrayView.h"
+#include "Containers/StringConv.h"
 #include "Misc/StringBuilder.h"
 #include "PlainPropsPrint.h"
 #include "Templates/UniquePtr.h"
@@ -15,6 +16,11 @@ struct FSchemaBatch;
 struct FEnumSchema;
 struct FStructSchema;
 
+inline TStringBuilderWithBuffer<TCHAR, 128> Print(FUtf8StringView View)
+{
+	return WriteToString<128>(FStringView(StringCast<TCHAR>(View.GetData(), View.Len())));
+}
+
 ///////////////////////////////////////////////////////////////////////////////
 
 struct FYamlBuilderDeleter
diff --git a/Engine/Plugins/Experimental/PlainProps/Source/Private/PlainPropsInternalTest.h b/Engine/Plugins/Experimental/PlainProps/Source/Private/PlainPropsInternalTest.h
--- a/Engine/Plugins/Experimental/PlainProps/Source/Private/PlainPropsInternalTest.h
+++ b/Engine/Plugins/Experimental/PlainProps/Source/Private/PlainPropsInternalTest.h
@@ -3,11 +3,15 @@
 #pragma once
 
 #include "PlainPropsTypes.h"
+#include "Containers/ContainersFwd.h"
 #include "Memory/MemoryFwd.h"
 
 namespace PlainProps
 {
 
-bool DiffSchemas(FMemoryView A, FMemoryView B);
+class FReadBatchId;
+struct FStructView;
+
+[[nodiscard]] FReadBatchId ParseBatch(TArray64<uint8>& OutData, TArray<FStructView>& OutObjects, FUtf8StringView YamlView);
 
 } // namespace PlainProps
diff --git a/Engine/Plugins/Experimental/PlainProps/Source/Private/PlainPropsParse.cpp b/Engine/Plugins/Experimental/PlainProps/Source/Private/PlainPropsParse.cpp
--- a/Engine/Plugins/Experimental/PlainProps/Source/Private/PlainPropsParse.cpp
+++ b/Engine/Plugins/Experimental/PlainProps/Source/Private/PlainPropsParse.cpp
@@ -10,8 +10,8 @@
 #include "PlainPropsBuildSchema.h"
 #include "PlainPropsInternalBuild.h"
 #include "PlainPropsInternalParse.h"
+#include "PlainPropsInternalPrint.h"
 #include "PlainPropsInternalText.h"
-#include "PlainPropsPrint.h"
 #include "PlainPropsWrite.h"
 #include <charconv>
 
@@ -25,14 +25,9 @@ void AppendString(FUtf8StringBuilderBase& Out, const FUtf8StringView& Str)
 	Out.Append(Str);
 }
 
-auto Print(FUtf8StringView View)
+TStringBuilderWithBuffer<TCHAR, 128> Print(FToken Token)
 {
-	return WriteToString<64>(FStringView(StringCast<TCHAR>(View.GetData(), View.Len())));
-}
-
-auto Print(FToken Token)
-{
-	return WriteToString<64>(FStringView(StringCast<TCHAR>(Token.Str, Token.Len)));
+	return WriteToString<128>(FStringView(StringCast<TCHAR>(Token.Str, Token.Len)));
 }
 
 inline const FParsedStructSchema& ResolveStructSchema(const FParsedSchemas& Schemas, FStructId Id)
@@ -172,31 +167,112 @@ TOptional<bool> Parse(FUtf8StringView String)
 	return NullOpt;
 }
 
-template<>
-TOptional<char8_t> Parse(FUtf8StringView String)
+static inline TOptional<uint32> ParseCodepoint4(FUtf8StringView HexString)
 {
-	if (String.Len() == 0)
+	check(HexString.Len() == 4);
+	uint32 Codepoint = 0;
+	for (int32 I = 0; I < 4; ++I)
+	{
+		Codepoint <<= 4;
+		UTF8CHAR C = HexString[I];
+		if (C >= '0' && C <= '9')
+		{
+			Codepoint += C - '0';
+		}
+		else if (C >= 'a' && C <= 'f')
+		{
+			Codepoint += C + 10 - 'a';
+		}
+		else if (C >= 'A' && C <= 'F')
+		{
+			Codepoint += C + 10 - 'A';
+		}
+		else
+		{
+			return NullOpt;
+		}
+	}
+	return Codepoint;
+}
+
+// Matches PlainPropsPrint EscapeChar()
+// The quotation mark (U+0022), backslash (U+005C),
+// and control characters U+0000 to U+001F are escaped (JSON Standard ECMA-404).
+// PP-TEXT: Missing support for Json two-character escape sequences \uxxxx\uxxxx
+// PP-TEXT: Missing support for Yaml escaping
+static inline TOptional<uint32> GrabEscapedCodepoint(FUtf8StringView& String)
+{
+	if (String.Len() < 2)
+	{
+		return NullOpt;
+	}
+	const UTF8CHAR* C = String.GetData();
+	if (*C++ != '\\')
+	{
+		return NullOpt;
+	}
+	int32 Len = 2;
+	TOptional<uint32> Codepoint;
+	switch (*C)
+	{
+		case '"':	Codepoint = *C;		break;
+		case '\\':	Codepoint = *C;		break;
+		case 'b':	Codepoint = '\b';	break;
+		case 'f':	Codepoint = '\f';	break;
+		case 'n':	Codepoint = '\n';	break;
+		case 'r':	Codepoint = '\r';	break;
+		case 't':	Codepoint = '\t';	break;
+		case 'u':
+			if (String.Len() >= 6)
+			{
+				if ((Codepoint = ParseCodepoint4(String.Mid(2, 4))))
+				{
+					Len = 6;
+					break;
+				}
+			}
+		default:
+			checkf(false, TEXT("Invalid escape sequence '%s'"), *Print(String));
+			break;
+	}
+	if (Codepoint)
 	{
-		return 0;
+		String.RightChopInline(Len);
 	}
+	return Codepoint;
+}
+
+template<>
+TOptional<char8_t> Parse(FUtf8StringView String)
+{
 	if (String.Len() == 1 && String[0] <= 127)
 	{
 		return static_cast<char8_t>(String[0]);
 	}
+	if (TOptional<uint32> Codepoint = GrabEscapedCodepoint(String))
+	{
+		if (Codepoint.GetValue() <= 127)
+		{
+			return static_cast<char8_t>(Codepoint.GetValue());
+		}
+	}
 	return NullOpt;
 }
 
 template<>
 TOptional<char16_t> Parse(FUtf8StringView String)
 {
-	if (String.Len() == 0)
-	{
-		return 0;
-	}
 	if (String.Len() == 1 && String[0] <= 127)
 	{
 		return static_cast<char16_t>(String[0]);
 	}
+	if (TOptional<uint32> Codepoint = GrabEscapedCodepoint(String))
+	{
+		if (Codepoint.GetValue() <= 0xFFFF)
+		{
+			return static_cast<char16_t>(Codepoint.GetValue());
+		}
+	}
 	UTF16CHAR Dest;
 	if (UTF16CHAR* DestEnd = FPlatformString::Convert(&Dest, 1, String.GetData(), String.Len()))
 	{
@@ -209,14 +285,17 @@ TOptional<char16_t> Parse(FUtf8StringView String)
 template<>
 TOptional<char32_t> Parse(FUtf8StringView String)
 {
-	if (String.Len() == 0)
-	{
-		return 0;
-	}
 	if (String.Len() == 1 && String[0] <= 127)
 	{
 		return static_cast<char32_t>(String[0]);
 	}
+	if (TOptional<uint32> Codepoint = GrabEscapedCodepoint(String))
+	{
+		if (Codepoint.GetValue() <= 0xFFFF)
+		{
+			return static_cast<char32_t>(Codepoint.GetValue());
+		}
+	}
 	// PP-TEXT: FPlatformString::Convert missing for UTF32CHAR in Parse()
 	checkf(false, TEXT("Missing conversion from UTF8CHAR to UTF32CHAR"));
 	return NullOpt;
@@ -226,7 +305,9 @@ TOptional<char32_t> Parse(FUtf8StringView String)
 
 constexpr FAsciiSet LinebreakSet("\r\n");
 constexpr FAsciiSet WhitespaceSet(" ");
-constexpr FAsciiSet TokenSet("'");
+constexpr FAsciiSet SingleQuoteSet("\'");
+constexpr FAsciiSet DoubleQuoteSet("\"");
+constexpr FAsciiSet BackslashSet("\\");
 constexpr FAsciiSet ScopeSet(".");
 constexpr FAsciiSet OpenRangeSet("(");
 constexpr FAsciiSet CloseRangeSet(")");
@@ -237,7 +318,6 @@ constexpr FAsciiSet NextParameterSet(",");
 constexpr FAsciiSet ParameterSet(OpenParameterSet | CloseParameterSet | NextParameterSet);
 
 constexpr bool PeekIsColon(FUtf8StringView View)			{ return *View.GetData() == ':'; }
-constexpr bool PeekIsToken(FUtf8StringView View)			{ return *View.GetData() == '\''; }
 constexpr bool IsRangeTokenLine(FUtf8StringView View)		{ return *View.GetData() == '-'; }
 constexpr bool IsEmptyOrCommentLine(FUtf8StringView View)	{ return View.IsEmpty() || *View.GetData() == '#'; }
 
@@ -652,20 +732,19 @@ TOptional<FToken> FYamlTokenizer::GrabBeginDocument()
 	}
 	check(ScopeStack.IsEmpty());
 	ScopeStack.Emplace(EScope::Document);
-	return FToken{ EToken::BeginDocument, /*Depth*/ 0, View};
+	return FToken{ EToken::BeginDocument, /*Depth*/ 0, View };
 }
 
 static FUtf8StringView GrabTokenFromLine(FLineInfo& Line)
 {
-	if (PeekIsToken(Line.View))
+	if (SingleQuoteSet.Contains(*Line.View.GetData()))
 	{
-		check(TokenSet.Contains(*Line.View.GetData()));
 		Line.View.RightChopInline(1);
 
-		const FUtf8StringView Token = FAsciiSet::FindPrefixWithout(Line.View, TokenSet);
+		const FUtf8StringView Token = FAsciiSet::FindPrefixWithout(Line.View, SingleQuoteSet);
 		Line.View.RightChopInline(Token.Len());
 
-		check(TokenSet.Contains(*Line.View.GetData()));
+		check(SingleQuoteSet.Contains(*Line.View.GetData()));
 		Line.View.RightChopInline(1);
 
 		const FUtf8StringView Skip = FAsciiSet::FindPrefixWith(Line.View, WhitespaceSet);
@@ -673,6 +752,40 @@ static FUtf8StringView GrabTokenFromLine(FLineInfo& Line)
 
 		return Token;
 	}
+	else if (DoubleQuoteSet.Contains(*Line.View.GetData()))
+	{
+		Line.View.RightChopInline(1);
+		int32 Len = 0;
+		const UTF8CHAR* Start = Line.View.GetData();
+
+		FUtf8StringView Part = FAsciiSet::FindPrefixWithout(Line.View, DoubleQuoteSet);
+		Line.View.RightChopInline(Part.Len());
+		Len += Part.Len();
+
+		// if the double quote " was preceeded by an uneven number of back slashes,
+		// then it was escaped and should be consumed as being part of the string
+		FUtf8StringView PartEscapeSuffix = FAsciiSet::FindSuffixWith(Part, BackslashSet);
+		while ((PartEscapeSuffix.Len() % 2))
+		{
+			check(DoubleQuoteSet.Contains(*Line.View.GetData()));
+			Line.View.RightChopInline(1);
+			Len += 1;
+
+			Part = FAsciiSet::FindPrefixWithout(Line.View, DoubleQuoteSet);
+			Line.View.RightChopInline(Part.Len());
+			Len += Part.Len();
+
+			PartEscapeSuffix = FAsciiSet::FindSuffixWith(Part, BackslashSet);
+		}
+
+		check(DoubleQuoteSet.Contains(*Line.View.GetData()));
+		Line.View.RightChopInline(1);
+
+		const FUtf8StringView Skip = FAsciiSet::FindPrefixWith(Line.View, WhitespaceSet);
+		Line.View.RightChopInline(Skip.Len());
+
+		return FUtf8StringView(Start, Len);
+	}
 	return {};
 }
 
@@ -831,7 +944,10 @@ void FYamlTokenizer::AdvanceReadLines()
 void FYamlTokenizer::SetLastError(FUtf8StringView Message)
 {
 	LastError.Reset();
-	LastError.Appendf("'%s' in line %u: '%s'", *Print(Message), ReadLine.LineNumber, *Print(ReadLine.FullView));
+	if (Message.Len())
+	{
+		LastError.Appendf("'%s' in line %u: '%s'", *Print(Message), ReadLine.LineNumber, *Print(ReadLine.FullView));
+	}
 }
 
 ////////////////////////////////////////////////////////////////////////////////
@@ -1411,7 +1527,13 @@ bool FMemberParser::ParseRange()
 bool FMemberParser::AdvanceToNextMember()
 {
 	const FUtf8StringView MemberName = MemberTokens.Peek()->Value();
-	const FMemberId MemberId = Names.NameMember(MemberName);
+	const FOptionalMemberId MemberId = Names.GetMemberId(MemberName);
+
+	if (!MemberId)
+	{
+		checkf(false, TEXT("Member name '%s' not found in any schema"), *Print(MemberName));
+		return false;
+	}
 
 	for (uint16 Idx = MemberIdx, N = MemberNames.Num(); Idx < N; ++Idx)
 	{
@@ -1432,7 +1554,7 @@ bool FMemberParser::AdvanceToNextMember()
 	}
 #endif
 
-	checkf(false, TEXT("Member '%s' not found"), *Print(MemberName));
+	checkf(false, TEXT("Member '%s' not found in struct schema"), *Print(MemberName));
 	return false;
 }
 
@@ -1606,6 +1728,7 @@ FTypedRange FMemberParser::ParseStructs(const FMemberSchema& RangeSchema, FScope
 	FTypedRange TypedRange { BuildRangeSchema };
 	if (int64 Num = Structs.Num())
 	{
+		check(Schema);
 		TypedRange.Values = FBuiltRange::Create(Scratch, Structs.Num(), sizeof(FBuiltStruct*));
 		FMemory::Memcpy(TypedRange.Values->Data, Structs.GetData(), Structs.NumBytes());
 	}		
@@ -1619,6 +1742,7 @@ FTypedRange FMemberParser::ParseRanges(const FMemberSchema& RangeSchema, FScoped
 	check(RangeSchema.NestedRangeTypes);
 
 	const ERangeSizeType NumType = RangeSchema.Type.AsRange().MaxSize;
+	const bool IsDynamic = RangeSchema.GetInnermostType() == DynamicStructType;
 	const FMemberSchema InnerRangeSchema
 	{
 		.Type = RangeSchema.NestedRangeTypes[0],
@@ -1628,11 +1752,20 @@ FTypedRange FMemberParser::ParseRanges(const FMemberSchema& RangeSchema, FScoped
 		.NestedRangeTypes = RangeSchema.NumInnerRanges > 2 ? RangeSchema.NestedRangeTypes + 1 : nullptr,
 	};
 
+	FMemberSchema BuiltInnerRangeSchema = InnerRangeSchema;
+
 	TArray64<FBuiltRange*> Ranges;
 	while (RangeTokens.HasMore())
 	{
 		FTypedRange TypedRange = ParseRangeInternal(InnerRangeSchema);
 		Ranges.Add(TypedRange.Values);
+		if (IsDynamic)
+		{
+			checkf(!BuiltInnerRangeSchema.InnerSchema ||
+				BuiltInnerRangeSchema.InnerSchema.Get() == TypedRange.Schema.InnerSchema,
+				TEXT("Heterogeneous struct ranges have not been implemented yet"));
+			BuiltInnerRangeSchema = TypedRange.Schema;
+		}
 	}
 
 	FBuiltRange* Out = nullptr;
@@ -1642,7 +1775,7 @@ FTypedRange FMemberParser::ParseRanges(const FMemberSchema& RangeSchema, FScoped
 		FMemory::Memcpy(Out->Data, Ranges.GetData(), Ranges.NumBytes());
 	}
 
-	return { MakeNestedRangeSchema(Scratch, NumType, InnerRangeSchema), Out };
+	return { MakeNestedRangeSchema(Scratch, NumType, BuiltInnerRangeSchema), Out };
 }
 
 static bool IsUnicodeString(const FMemberSchema& MemberSchema)
@@ -1693,27 +1826,57 @@ FBuiltStruct* FMemberParser::ParseMembersInternal(
 	return MemberParser.ParseMembers(Schema);
 }
 
+static FUtf8StringView ParseString(TArray<UTF8CHAR>& Out, FUtf8StringView String)
+{
+	FUtf8StringView Verbatim = FAsciiSet::FindPrefixWithout(String, BackslashSet);
+	if (Verbatim.Len() == String.Len())
+	{
+		return String;
+	}
+
+	Out.Reserve(String.Len());
+	while (!String.IsEmpty())
+	{
+		Out.Append(MakeArrayView(Verbatim.GetData(), Verbatim.Len()));
+		String.RightChopInline(Verbatim.Len());
+		while (TOptional<uint32> Codepoint = GrabEscapedCodepoint(String))
+		{
+			checkf(Codepoint.GetValue() <= 127, TEXT("Unexpected codepoint: %u"), Codepoint.GetValue());
+			Out.Add(static_cast<UTF8CHAR>(Codepoint.GetValue()));
+		}
+		Verbatim = FAsciiSet::FindPrefixWithout(String, BackslashSet);
+	}
+	return { Out.GetData(), Out.Num() };
+}
+
 FTypedRange FMemberParser::ParseUnicodeLeafValueAsRange(
 	ERangeSizeType NumType,
 	FUnpackedLeafType Leaf,
 	FUtf8StringView String)
 {
+	TArray<UTF8CHAR> Buffer;
+	FUtf8StringView Parsed = ParseString(Buffer, String);
 	if (Leaf.Width == ELeafWidth::B8)
 	{
-		const char8_t* Src = reinterpret_cast<const char8_t*>(String.GetData());
-		return BuildLeafRange(Scratch, NumType, MakeArrayView(Src, String.Len()));
+		const char8_t* Src = reinterpret_cast<const char8_t*>(Parsed.GetData());
+		return BuildLeafRange(Scratch, NumType, MakeArrayView(Src, Parsed.Len()));
 	}
 	else if (Leaf.Width == ELeafWidth::B16)
 	{
-		auto Utf16 = StringCast<UTF16CHAR>(String.GetData(), String.Len());
-		const char16_t* Src = reinterpret_cast<const char16_t*>(Utf16.Get());
-		return BuildLeafRange(Scratch, NumType, MakeArrayView(Src, Utf16.Length()));
+		const int32 DstLen = FPlatformString::ConvertedLength<UTF16CHAR>(Parsed.GetData(), Parsed.Len());
+		TArray<UTF16CHAR, TInlineAllocator<1024>> Buf;
+		Buf.Reserve(DstLen);
+		UTF16CHAR* Dst = Buf.GetData();
+		const UTF16CHAR* DstEnd = FPlatformString::Convert(Dst, DstLen, Parsed.GetData(), Parsed.Len());
+		check(DstEnd);
+		check(DstEnd - Dst == DstLen);
+		return BuildLeafRange(Scratch, NumType, MakeArrayView(reinterpret_cast<const char16_t*>(Dst), DstLen));
 	}
 	else // ELeafWidth::B32
 	{
 		check(Leaf.Width == ELeafWidth::B32);
 		// PP-TEXT: StringCast<UTF32CHAR>/ToUtf32() missing in ParseUnicodeLeafValueAsRange()
-		// auto Utf32 = StringCast<UTF32CHAR>(String.GetData(), String.Len());
+		// auto Utf32 = StringCast<UTF32CHAR>(Parsed.GetData(), Parsed.Len());
 		// const char32_t* Src = reinterpret_cast<const char32_t*>(Utf32.Get());
 		// return BuildLeafRange(Scratch, NumType, MakeArrayView(Src, Utf32.Length()));
 		checkf(false, TEXT("StringCast<UTF32CHAR>/ToUtf32() not implemented"));
@@ -1725,17 +1888,13 @@ FTypedRange FMemberParser::ParseUnicodeLeafValueAsRange(
 
 void FBatchParser::Parse(TArray64<uint8>& Out)
 {
+	Tokenize();
 	ParseAll();
 	Write(Out);
 }
 
-void FBatchParser::ParseAll()
+void FBatchParser::Tokenize()
 {
-	int64 EnumsIdx = -1;
-	int64 StructsIdx = -1;
-	int64 ObjectsIdx = -1;
-
-	TArray64<FToken> Tokens;
 	while (Tokenizer.HasMore())
 	{
 		TOptional<FToken> Token = Tokenizer.GrabToken();
@@ -1767,7 +1926,10 @@ void FBatchParser::ParseAll()
 	checkf(EnumsIdx >= 0, TEXT("No '%.*hs' section found"), GLiterals.Structs.Len(), GLiterals.Structs.GetData());
 	checkf(StructsIdx >= 0, TEXT("No '%.*hs' section found"), GLiterals.Enums.Len(), GLiterals.Enums.GetData());
 	checkf(ObjectsIdx >= 0, TEXT("No '%.*hs' section found"), GLiterals.Objects.Len(), GLiterals.Objects.GetData());
+}
 
+void FBatchParser::ParseAll()
+{
 	TConstArrayView64<FToken> TokensView(Tokens);
 	ParseEnumSchemas(TokensView.RightChop(EnumsIdx));
 	ParseStructSchemas(TokensView.RightChop(StructsIdx));
@@ -1776,10 +1938,10 @@ void FBatchParser::ParseAll()
 
 void FBatchParser::ParseStructSchemas(TConstArrayView64<FToken> TokensView)
 {
-	FTokenReader Tokens(TokensView);
-	verify(Tokens.GrabToken().Value() == GLiterals.Structs);
+	FTokenReader TokenIt(TokensView);
+	verify(TokenIt.GrabToken().Value() == GLiterals.Structs);
 
-	FStructSchemaParser Parser(Tokens, Names, Types, Schemas);
+	FStructSchemaParser Parser(TokenIt, Names, Types, Schemas);
 	while (Parser.HasMore())
 	{
 		FParsedStructSchema ParsedStructSchema;
@@ -1793,10 +1955,10 @@ void FBatchParser::ParseStructSchemas(TConstArrayView64<FToken> TokensView)
 
 void FBatchParser::ParseEnumSchemas(TConstArrayView64<FToken> TokensView)
 {
-	FTokenReader Tokens(TokensView);
-	verify(Tokens.GrabToken().Value() == GLiterals.Enums);
+	FTokenReader TokenIt(TokensView);
+	verify(TokenIt.GrabToken().Value() == GLiterals.Enums);
 
-	FEnumSchemaParser Parser(Tokens, Names, Types);
+	FEnumSchemaParser Parser(TokenIt, Names, Types);
 	while (Parser.HasMore())
 	{
 		FParsedEnumSchema ParsedEnumSchema;
@@ -1810,10 +1972,10 @@ void FBatchParser::ParseEnumSchemas(TConstArrayView64<FToken> TokensView)
 
 void FBatchParser::ParseObjects(TConstArrayView64<FToken> TokensView)
 {
-	FTokenReader Tokens(TokensView);
-	verify(Tokens.GrabToken().Value() == GLiterals.Objects);
+	FTokenReader TokenIt(TokensView);
+	verify(TokenIt.GrabToken().Value() == GLiterals.Objects);
 
-	FScopedTokenReader ScopedReader(Tokens);
+	FScopedTokenReader ScopedReader(TokenIt);
 	while (ScopedReader.HasMore())
 	{
 		FToken Token;
@@ -1836,14 +1998,14 @@ void FBatchParser::ParseObjects(TConstArrayView64<FToken> TokensView)
 		const FParsedStructSchema& ParsedSchema = ResolveStructSchema(Schemas, StructId.Get());
 
 		{
-			FScopedTokenReader MemberTokens(Tokens, Token);
+			FScopedTokenReader MemberTokens(TokenIt, Token);
 			FMemberParser MemberParser(Scratch, ParsedSchema, MemberTokens, Names, Types, Schemas);
 			FBuiltStruct* BuiltStruct = MemberParser.ParseMembers(StructId.Get());
 			check(!MemberTokens.HasMore());
 			Objects.Emplace(StructId.Get(), BuiltStruct);
 		}
 
-		Token = Tokens.GrabToken(EToken::EndStruct);
+		Token = TokenIt.GrabToken(EToken::EndStruct);
 	}
 }
 
@@ -1868,7 +2030,25 @@ void FBatchParser::Write(TArray64<uint8>& Out)
 	FBuiltSchemas BuiltSchemas = SchemaBuilders.Build(); 
 
 	FWriter Writer(Names, Bindings, BuiltSchemas, ESchemaFormat::StableNames);
-	Writer.WriteSchemas(Out);
+	TArray64<uint8> Tmp;
+
+	// Write schemas
+	Writer.WriteSchemas(Tmp);
+	WriteInt(Out, IntCastChecked<uint32>(Tmp.Num()));
+	WriteArray(Out, Tmp);
+	Tmp.Reset();
+
+	// Write objects
+	for (const TPair<FStructId, FBuiltStruct*>& Object : Objects)
+	{
+		WriteInt(Tmp, Writer.GetWriteId(Object.Key).Get().Idx);
+		Writer.WriteMembers(Tmp, Object.Key, *Object.Value);
+		WriteSkippableSlice(Out, Tmp);
+		Tmp.Reset();
+	}
+
+	// Write object terminator
+	WriteSkippableSlice(Out, TConstArrayView64<uint8>());
 }
 
 ///////////////////////////////////////////////////////////////////////////////
@@ -1881,4 +2061,4 @@ void ParseYamlBatch(TArray64<uint8>& OutBinary, FUtf8StringView Yaml)
 	BatchParser.Parse(OutBinary);
 }
 
-} // namespace PlainProps
\ No newline at end of file
+} // namespace PlainProps
diff --git a/Engine/Plugins/Experimental/PlainProps/Source/Private/PlainPropsPrint.cpp b/Engine/Plugins/Experimental/PlainProps/Source/Private/PlainPropsPrint.cpp
--- a/Engine/Plugins/Experimental/PlainProps/Source/Private/PlainPropsPrint.cpp
+++ b/Engine/Plugins/Experimental/PlainProps/Source/Private/PlainPropsPrint.cpp
@@ -10,7 +10,9 @@
 #include "PlainPropsInternalPrint.h"
 #include "PlainPropsInternalRead.h"
 #include "PlainPropsInternalText.h"
+#include "Misc/AsciiSet.h"
 #include "Misc/StringBuilder.h"
+#include <charconv>
 
 namespace PlainProps
 {
@@ -36,6 +38,16 @@ FAnsiStringView ToString(ELeafWidth Width)
 	return GLiterals.Widths[(uint8)Width];
 }
 
+FAnsiStringView ToString(ESchemaFormat Format)
+{
+	switch (Format)
+	{
+		case ESchemaFormat::InMemoryNames:	return "InMemoryNames";
+		case ESchemaFormat::StableNames:	return "StableNames";
+	}
+	return "Unknown";
+}
+
 //////////////////////////////////////////////////////////////////////////
 
 void FIdIndexerBase::InitParameterNames()
@@ -54,6 +66,102 @@ void FIdIndexerBase::InitParameterNames()
 	}
 }
 
+///////////////////////////////////////////////////////////////////////////////
+// Escape the quotation mark (U+0022), backslash (U+005C),
+// and control characters U+0000 to U+001F (JSON Standard ECMA-404)
+constexpr FAsciiSet EscapeSet("\\\""
+	"\x00\x01\x02\x03\x04\x05\x06\x07\x08\x09\x0a\x0b\x0c\x0d\x0e\x0f"
+	"\x10\x11\x12\x13\x14\x15\x16\x17\x18\x19\x1a\x1b\x1c\x1d\x1e\x1f");
+
+static inline void EscapeChar(FUtf8Builder& Out, UTF8CHAR Char)
+{
+	switch (Char)
+	{
+	case '\"': Out.Append("\\\"");	break;
+	case '\\': Out.Append("\\\\");	break;
+	case '\b': Out.Append("\\b");	break;
+	case '\f': Out.Append("\\f");	break;
+	case '\n': Out.Append("\\n");	break;
+	case '\r': Out.Append("\\r");	break;
+	case '\t': Out.Append("\\t");	break;
+	default:
+		Out.Appendf(UTF8TEXT("\\u%04x"), uint32(Char));
+		break;
+	}
+}
+
+template<Enumeration T>
+void Print(FUtf8Builder& Out, T Value)
+{
+	Out.Append(ToString(Value));
+}
+
+template<Arithmetic T> requires (std::is_integral_v<T>)
+void Print(FUtf8Builder& Out, T Value)
+{
+	constexpr size_t BufSize = std::numeric_limits<T>::digits10 + 3;
+	char Buf[BufSize]{};
+	std::to_chars_result result = std::to_chars(Buf, Buf + BufSize, Value);
+	check(result.ec == std::errc());
+	Out.Append(Buf);
+}
+
+template<Arithmetic T> requires (std::is_floating_point_v<T>)
+void Print(FUtf8Builder& Out, T Value)
+{
+	constexpr size_t BufSize = 32;
+	char Buf[BufSize]{};
+	std::to_chars_result result = std::to_chars(Buf, Buf + BufSize, Value, std::chars_format::general);
+	check(result.ec == std::errc());
+	Out.Append(Buf);
+}
+
+template<>
+void Print(FUtf8Builder& Out, bool Value)
+{
+	Out.Append(Value ? GLiterals.True : GLiterals.False);
+}
+
+template<>
+void Print(FUtf8Builder& Out, char8_t Value)
+{
+	UTF8CHAR Char = static_cast<UTF8CHAR>(Value);
+	if (EscapeSet.Contains(Char))
+	{
+		EscapeChar(Out, Char);
+	}
+	else
+	{
+		Out.AppendChar(Char);
+	}
+}
+
+template<>
+void Print(FUtf8Builder& Out, char16_t Value)
+{
+	if (Value <= 127)
+	{
+		Print(Out, static_cast<char8_t>(Value));
+	}
+	else
+	{
+		Out.Append(FUtf8StringView(StringCast<UTF8CHAR>(reinterpret_cast<UTF16CHAR*>(&Value), 1)));
+	}
+}
+
+template<>
+void Print(FUtf8Builder& Out, char32_t Value)
+{
+	if (Value <= 127)
+	{
+		Print(Out, static_cast<char8_t>(Value));
+	}
+	else
+	{
+		Out.Append(FUtf8StringView(StringCast<UTF8CHAR>(reinterpret_cast<UTF32CHAR*>(&Value), 1)));
+	}
+}
+
 ///////////////////////////////////////////////////////////////////////////////
 
 static void PrintRangeType(FUtf8Builder& Out, FRangeType Type)
@@ -246,6 +354,10 @@ private:
 	void AppendNewLine();
 	void AppendIndentation();
 	void AppendIdentifier(FUtf8StringView Id);
+	template<typename T>
+	void AppendValue(T Value);
+	template<>
+	void AppendValue(FUtf8StringView Value);
 
 	struct FStackInfo
 	{
@@ -545,7 +657,8 @@ void FYamlBuilder::AddLeafId(FUtf8StringView Id)
 template<typename T>
 void FYamlBuilder::AddLeafValue(T Value)
 {
-	Text << " '" << Value << '\'';
+	Text.AppendChar(' ');
+	AppendValue(Value);
 	IsNewLine = false;
 	Stack.Last().IsEmpty = false;
 }
@@ -555,7 +668,7 @@ void FYamlBuilder::AddLeaf(T Value)
 {
 	AppendNewLine();
 	AppendIndentation();
-	Text << '\'' << Value << '\'';
+	AppendValue(Value);
 	IsNewLine = false;
 	Stack.Last().IsEmpty = false;
 }
@@ -577,9 +690,9 @@ void FYamlBuilder::AppendNewLine()
 
 void FYamlBuilder::AppendIndentation()
 {
-	for (uint32 I = 0; I < IndentLevel; ++I)
+	for (uint32 I = 0; I < 2*IndentLevel; ++I)
 	{
-		Text << "  ";
+		Text.AppendChar(' ');
 	}
 	if (!Stack.Last().IsInStruct)
 	{
@@ -588,12 +701,52 @@ void FYamlBuilder::AppendIndentation()
 	IsNewLine = false;
 }
 
+static void PrintQuotedString(FUtf8Builder& Out, FUtf8StringView Value)
+{
+	FUtf8StringView Verbatim = FAsciiSet::FindPrefixWithout(Value, EscapeSet | "'");
+	if (Verbatim.Len() == Value.Len())
+	{
+		Out << '\'' << Value << '\'';
+		return;
+	}
+
+	Out << '\"';
+	while (!Value.IsEmpty())
+	{
+		Out << Verbatim;
+		Value.RightChopInline(Verbatim.Len());
+		FUtf8StringView Escape = FAsciiSet::FindPrefixWith(Value, EscapeSet);
+		for (UTF8CHAR Char : Escape)
+		{
+			EscapeChar(Out, Char);
+		}
+		Value.RightChopInline(Escape.Len());
+		Verbatim = FAsciiSet::FindPrefixWithout(Value, EscapeSet);
+	}
+	Out << '\"';
+}
+
 void FYamlBuilder::AppendIdentifier(FUtf8StringView Id)
 {
-	Text << '\'' << Id << "':";
+	PrintQuotedString(Text, Id);
+	Text.AppendChar(':');
 	IsNewLine = false;
 }
 
+template<typename T>
+void FYamlBuilder::AppendValue(T Value)
+{
+	Text.AppendChar('\'');
+	Print(Text, Value);
+	Text.AppendChar('\'');
+}
+
+template<>
+void FYamlBuilder::AppendValue(FUtf8StringView Value)
+{
+	PrintQuotedString(Text, Value);
+}
+
 ///////////////////////////////////////////////////////////////////////////////
 
 FBatchPrinter::FBatchPrinter(FYamlBuilder& InTextBuilder, const FBatchIds& InIds)
@@ -684,7 +837,7 @@ void FBatchPrinter::PrintEnumSchema(const FEnumSchema& Enum)
 {
 	TextBuilder.BeginStruct(*FPrintId<128>(Ids, Enum.Type));
 	TextBuilder.AddLeaf(GLiterals.FlagMode, !!Enum.FlagMode);
-	TextBuilder.AddLeaf(GLiterals.Width, ToString(Enum.Width));
+	TextBuilder.AddLeaf(GLiterals.Width, Enum.Width);
 
 	TextBuilder.BeginRange(GLiterals.Constants);
 	TConstArrayView<FNameId> EnumNames = MakeConstArrayView(Enum.Footer, Enum.Num);
@@ -801,8 +954,13 @@ void FMemberPrinter::PrintLeaf(FMemberId Id, FLeafView LeafView)
 		}
 		break;
 	case ELeafType::Unicode:
-		PrintUnicodeLeafValue(LeafView);
-		break;
+		switch (LeafView.Leaf.Width)
+		{
+			case ELeafWidth::B8:	TextBuilder.AddLeafValue(LeafView.AsChar8());	break;
+			case ELeafWidth::B16:	TextBuilder.AddLeafValue(LeafView.AsChar16());	break;
+			case ELeafWidth::B32:	TextBuilder.AddLeafValue(LeafView.AsChar32());	break;
+			case ELeafWidth::B64:	check(false);									break;
+		};
 	}
 
 	PrintSchemaComment(LeafView.Leaf, LeafView.Enum);
@@ -979,49 +1137,26 @@ bool FMemberPrinter::IsUnicodeString(const FRangeView& RangeView)
 	return RangeView.Num() > 0 && Type.IsLeaf() && Type.AsLeaf().Type == ELeafType::Unicode;
 }
 
-template <typename CharType, typename LeafCharType>
-static void AddUnicodeLeaf(FYamlBuilder& TextBuilder, LeafCharType Char)
-{
-	if (Char == 0)
-	{
-		TextBuilder.AddLeafValue(FUtf8StringView{});
-	}
-	else if (Char <= 127)
-	{
-		TextBuilder.AddLeafValue(static_cast<UTF8CHAR>(Char));
-	}
-	else
-	{
-		CharType Temp = static_cast<CharType>(Char);
-		TextBuilder.AddLeafValue(FUtf8StringView(StringCast<UTF8CHAR>(&Temp, 1)));
-	}
-}
-
-void FMemberPrinter::PrintUnicodeLeafValue(FLeafView LeafView)
-{
-	switch (LeafView.Leaf.Width)
-	{
-		case ELeafWidth::B8:	AddUnicodeLeaf<UTF8CHAR >(TextBuilder, LeafView.AsChar8());		break;
-		case ELeafWidth::B16:	AddUnicodeLeaf<UTF16CHAR>(TextBuilder, LeafView.AsChar16());	break;
-		case ELeafWidth::B32:	AddUnicodeLeaf<UTF32CHAR>(TextBuilder, LeafView.AsChar32());	break;
-		case ELeafWidth::B64:	check(false);													break;
-	};
-}
-
 template <typename CharType, typename RangeCharType>
-static void AddUnicodeRangeLeaf(FYamlBuilder& TextBuilder, FOptionalMemberId Id, TRangeView<RangeCharType> Utf)
-{
-	check(Utf.Num() > 0);
-	const CharType* Src = reinterpret_cast<const CharType*>(Utf.begin());
-	int32 Len = IntCastChecked<int32>(Utf.Num());
-	Len -= (Utf[Len-1] == CharType(0)); // strip null terminator (for utf8->utf8 where StringCast is pass through)
+static void AddUnicodeRangeLeaf(FYamlBuilder& TextBuilder, FOptionalMemberId Id, TRangeView<RangeCharType> Range)
+{
+	check(Range.Num() > 0);
+	const CharType* Src = reinterpret_cast<const CharType*>(Range.begin());
+	const int32 SrcLen = IntCastChecked<int32>(Range.Num());
+	const int32 DstLen = FPlatformString::ConvertedLength<UTF8CHAR>(Src, SrcLen);
+	TArray<UTF8CHAR, TInlineAllocator<1024>> Buf;
+	Buf.Reserve(DstLen);
+	UTF8CHAR* Dst = Buf.GetData();
+	const UTF8CHAR* DstEnd = FPlatformString::Convert(Dst, DstLen, Src, SrcLen);
+	check(DstEnd);
+	check(DstEnd - Dst == DstLen);
 	if (Id)
 	{
-		TextBuilder.AddLeafValue(FUtf8StringView(StringCast<UTF8CHAR>(Src, Len)));
+		TextBuilder.AddLeafValue(FUtf8StringView(Dst, DstLen));
 	}
 	else
 	{
-		TextBuilder.AddLeaf(FUtf8StringView(StringCast<UTF8CHAR>(Src, Len)));
+		TextBuilder.AddLeaf(FUtf8StringView(Dst, DstLen));
 	}
 }
 
@@ -1299,4 +1434,42 @@ void PrintDiff(FUtf8Builder& Out, const FIds& Ids, const FDiffPath& Diff)
 	Out << ')';
 }
 
+void PrintDiff(FUtf8Builder& Out, const FBatchIds& Ids, const FReadDiffPath& Diff)
+{
+	check(Diff.Num());
+
+	bool bWasName = false;
+	// print type name for the outermost struct
+	if (Diff.Last().Struct)
+	{
+		Ids.AppendString(Out, Ids.Resolve(Diff.Last().Struct.Get()).Name);
+		bWasName = true;
+	}
+	// print struct members path with range indices
+	for (FReadDiffNode Node : ReverseIterate(Diff))
+	{
+		if (Node.Name || Node.RangeIdx == ~0u)
+		{
+			if (bWasName)
+			{
+				Out << '.';
+			}
+			if (!Node.Name)
+			{
+				Out << GLiterals.Super;
+			}
+			else
+			{
+				Ids.AppendString(Out, Node.Name);
+			}
+			bWasName = true;
+		}
+		else if (Node.Type.IsRange())
+		{
+			Out << "[" << Node.RangeIdx << "]";
+			bWasName = false;
+		}
+	}
+}
+
 } // namespace PlainProps
diff --git a/Engine/Plugins/Experimental/PlainProps/Source/Public/PlainPropsDiff.h b/Engine/Plugins/Experimental/PlainProps/Source/Public/PlainPropsDiff.h
--- a/Engine/Plugins/Experimental/PlainProps/Source/Public/PlainPropsDiff.h
+++ b/Engine/Plugins/Experimental/PlainProps/Source/Public/PlainPropsDiff.h
@@ -36,4 +36,19 @@ struct FDiffContext
 
 [[nodiscard]] PLAINPROPS_API bool DiffStructs(FStructId BindId, const void* A, const void* B, FDiffContext& Ctx);
 
-} // namespace PlainProps
\ No newline at end of file
+////////////////////////////////////////////////////////////////////////////////
+
+struct FReadDiffNode
+{
+	FMemberType				Type;
+	FOptionalStructSchemaId	Struct;
+	FOptionalMemberId		Name;
+	uint64					RangeIdx = ~0u; 
+};
+
+struct FReadDiffPath : public TArray<FReadDiffNode, TInlineAllocator<16>> {};
+
+bool PLAINPROPS_API DiffStruct(FStructView A, FStructView B, FReadDiffPath& Out);
+bool PLAINPROPS_API DiffSchemas(FReadBatchId A, FReadBatchId B);
+
+} // namespace PlainProps
diff --git a/Engine/Plugins/Experimental/PlainProps/Source/Public/PlainPropsPrint.h b/Engine/Plugins/Experimental/PlainProps/Source/Public/PlainPropsPrint.h
--- a/Engine/Plugins/Experimental/PlainProps/Source/Public/PlainPropsPrint.h
+++ b/Engine/Plugins/Experimental/PlainProps/Source/Public/PlainPropsPrint.h
@@ -12,16 +12,44 @@ class FBatchIds;
 struct FDiffPath;
 struct FMemberSchema;
 class FReadBatchId;
+struct FReadDiffPath;
 struct FStructView;
 
+PLAINPROPS_API FAnsiStringView ToString(ELeafWidth Width);
 PLAINPROPS_API FAnsiStringView ToString(ERangeSizeType Width);
+PLAINPROPS_API FAnsiStringView ToString(ESchemaFormat Format);
 PLAINPROPS_API FAnsiStringView ToString(FUnpackedLeafType Leaf);
-PLAINPROPS_API FAnsiStringView ToString(ELeafWidth Width);
+
+///////////////////////////////////////////////////////////////////////////////
+
+template<typename T> inline void Print(FUtf8Builder& Out, T Value)
+{
+	static_assert(!sizeof(T), "Unsupported type for Print");
+}
+
+template<> PLAINPROPS_API void Print(FUtf8Builder& Out, ERangeSizeType Value);
+template<> PLAINPROPS_API void Print(FUtf8Builder& Out, ELeafWidth Value);
+
+template<> PLAINPROPS_API void Print(FUtf8Builder& Out, bool Value);
+template<> PLAINPROPS_API void Print(FUtf8Builder& Out, int8 Value);
+template<> PLAINPROPS_API void Print(FUtf8Builder& Out, int16 Value);
+template<> PLAINPROPS_API void Print(FUtf8Builder& Out, int32 Value);
+template<> PLAINPROPS_API void Print(FUtf8Builder& Out, int64 Value);
+template<> PLAINPROPS_API void Print(FUtf8Builder& Out, uint8 Value);
+template<> PLAINPROPS_API void Print(FUtf8Builder& Out, uint16 Value);
+template<> PLAINPROPS_API void Print(FUtf8Builder& Out, uint32 Value);
+template<> PLAINPROPS_API void Print(FUtf8Builder& Out, uint64 Value);
+template<> PLAINPROPS_API void Print(FUtf8Builder& Out, float Value);
+template<> PLAINPROPS_API void Print(FUtf8Builder& Out, double Value);
+template<> PLAINPROPS_API void Print(FUtf8Builder& Out, char8_t Value);
+template<> PLAINPROPS_API void Print(FUtf8Builder& Out, char16_t Value);
+template<> PLAINPROPS_API void Print(FUtf8Builder& Out, char32_t Value);
 
 ///////////////////////////////////////////////////////////////////////////////
 
 PLAINPROPS_API void PrintMemberSchema(FUtf8Builder& Out, const FIds& Ids, FMemberSchema Member);
 PLAINPROPS_API void PrintYamlBatch(FUtf8Builder& Out, const FBatchIds& Ids, TConstArrayView<FStructView> Objects);
 PLAINPROPS_API void PrintDiff(FUtf8Builder& Out, const FIds& Ids, const FDiffPath& Diff);
+PLAINPROPS_API void PrintDiff(FUtf8Builder& Out, const FBatchIds& Ids, const FReadDiffPath& Diff);
 
 } // namespace PlainProps
diff --git a/Engine/Plugins/Experimental/PlainProps/Source/Public/PlainPropsRead.h b/Engine/Plugins/Experimental/PlainProps/Source/Public/PlainPropsRead.h
--- a/Engine/Plugins/Experimental/PlainProps/Source/Public/PlainPropsRead.h
+++ b/Engine/Plugins/Experimental/PlainProps/Source/Public/PlainPropsRead.h
@@ -298,6 +298,7 @@ class FRangeView
 public:
 	uint64							Num() const				{ return NumItems;}
 	bool							IsEmpty() const			{ return NumItems == 0; }
+	FMemberType						GetItemType() const		{ return Schema.ItemType; }
 	bool							IsLeafRange() const		{ return Schema.ItemType.GetKind() == EMemberKind::Leaf; }
 	bool							IsStructRange() const	{ return Schema.ItemType.GetKind() == EMemberKind::Struct; }
 	bool							IsNestedRange() const	{ return Schema.ItemType.GetKind() == EMemberKind::Range; }
@@ -343,6 +344,13 @@ public:
 		return *this;
 	}
 
+	FBoolRangeIterator operator++(int)
+	{
+		FBoolRangeIterator Tmp(*this);
+		operator++();
+		return Tmp;
+	}
+
 	bool operator!=(FBoolRangeIterator Rhs)
 	{
 		return Byte != Rhs.Byte || Mask != Rhs.Mask;
diff --git a/Engine/Plugins/Experimental/PlainProps/Source/Tests/PlainPropsTests.cpp b/Engine/Plugins/Experimental/PlainProps/Source/Tests/PlainPropsTests.cpp
--- a/Engine/Plugins/Experimental/PlainProps/Source/Tests/PlainPropsTests.cpp
+++ b/Engine/Plugins/Experimental/PlainProps/Source/Tests/PlainPropsTests.cpp
@@ -9,7 +9,9 @@
 #include "Containers/Set.h"
 #include "Containers/StringConv.h"
 #include "Containers/StringView.h"
+#include "PlainPropsDiff.h"
 #include "PlainPropsRead.h"
+#include "PlainPropsInternalBuild.h"
 #include "PlainPropsWrite.h"
 #include "PlainPropsBuildSchema.h"
 #include "PlainPropsIndex.h"
@@ -18,7 +20,6 @@
 #include "PlainPropsInternalParse.h"
 #include "PlainPropsInternalPrint.h"
 #include "PlainPropsInternalRead.h"
-#include "PlainPropsInternalTest.h"
 #include "Templates/UnrealTemplate.h"
 
 DEFINE_LOG_CATEGORY_STATIC(LogPlainPropsTests, Log, All);
@@ -33,146 +34,6 @@ static bool operator==(FParametricTypeView A, FParametricTypeView B)
 	return A.Name == B.Name && A.NumParameters == B.NumParameters && !FMemory::Memcmp(A.Parameters, B.Parameters, A.NumParameters * sizeof(FType));
 }
 
-static uint64 CalculateSize(const FSchemaBatch& In)
-{
-	uint32 NumParameters = 0;
-	for (FParametricType ParametricType : In.GetParametricTypes())
-	{
-		NumParameters += ParametricType.Parameters.NumParameters;
-	}
-	const uint8* End = reinterpret_cast<const uint8*>(In.GetFirstParameter() + NumParameters);
-	const uint64 Size = End - reinterpret_cast<const uint8*>(&In);
-	return Size;
-}
-
-static bool operator==(const FSchemaBatch& A, const FSchemaBatch& B)
-{
-	const uint64 SizeA = CalculateSize(A);
-	const uint64 SizeB = CalculateSize(B);
-	return SizeA == SizeB && FMemory::Memcmp(&A, &B, SizeA) == 0;
-}
-
-static bool operator==(const FStructSchema& A, const FStructSchema& B)
-{
-	const uint64 SizeA = CalculateSize(A);
-	const uint64 SizeB = CalculateSize(B);
-	return SizeA == SizeB && FMemory::Memcmp(&A, &B, SizeA) == 0;
-}
-
-static bool operator==(const FEnumSchema& A, const FEnumSchema& B)
-{
-	const uint64 SizeA = CalculateSize(A);
-	const uint64 SizeB = CalculateSize(B);
-	return SizeA == SizeB && FMemory::Memcmp(&A, &B, SizeA) == 0;
-}
-
-static bool operator!=(const FSchemaBatch& A, const FSchemaBatch& B)	{ return !(A == B); }
-static bool operator!=(const FStructSchema& A, const FStructSchema& B)	{ return !(A == B); }
-static bool operator!=(const FEnumSchema& A, const FEnumSchema& B)		{ return !(A == B); }
-
-////////////////////////////////////////////////////////////////////////////////
-
-static bool DiffStructSchema(const FStructSchema& A, const FStructSchema& B)
-{
-	const bool Diff = A != B;
-#if DO_CHECK
-	CHECK(A.Type == B.Type);
-	CHECK(A.NumMembers == B.NumMembers);
-	CHECK(A.NumRangeTypes == B.NumRangeTypes);
-	CHECK(A.NumInnerSchemas == B.NumInnerSchemas);
-	CHECK(A.Inheritance == B.Inheritance);
-	CHECK(A.IsDense == B.IsDense);
-	CHECK(EqualItems(A.GetMemberTypes(), B.GetMemberTypes()));
-	CHECK(EqualItems(A.GetRangeTypes(), B.GetRangeTypes()));
-	CHECK(EqualItems(A.GetMemberNames(), B.GetMemberNames()));
-	CHECK(EqualItems(
-		MakeArrayView(A.GetInnerSchemas(), A.NumInnerSchemas),
-		MakeArrayView(B.GetInnerSchemas(), B.NumInnerSchemas)));
-	CHECK(A.GetSuperSchema() == B.GetSuperSchema());
-	CHECK(A == B);
-#endif
-	return Diff;
-}
-
-static bool DiffEnumSchema(const FEnumSchema& A, const FEnumSchema& B)
-{
-	const bool Diff = A != B;
-#if DO_CHECK
-	CHECK(A.Type == B.Type);
-	CHECK(A.FlagMode == B.FlagMode);
-	CHECK(A.ExplicitConstants == B.ExplicitConstants);
-	CHECK(A.Width == B.Width);
-	CHECK(A.Num == B.Num);
-	CHECK(EqualItems(MakeArrayView(A.Footer, A.Num), MakeArrayView(B.Footer, B.Num)));
-	if (A.Width == B.Width)
-	{
-		switch (A.Width)
-		{
-		case ELeafWidth::B8:  CHECK(EqualItems(GetConstants<uint8 >(A), GetConstants<uint8 >(B))); break;
-		case ELeafWidth::B16: CHECK(EqualItems(GetConstants<uint16>(A), GetConstants<uint16>(B))); break;
-		case ELeafWidth::B32: CHECK(EqualItems(GetConstants<uint32>(A), GetConstants<uint32>(B))); break;
-		case ELeafWidth::B64: CHECK(EqualItems(GetConstants<uint64>(A), GetConstants<uint64>(B))); break;
-		}
-	}
-	CHECK(A == B);
-#endif
-	return Diff;
-}
-
-static bool DiffSchemas(const FSchemaBatch& A, const FSchemaBatch& B)
-{
-	const bool Diff = A != B;
-#if DO_CHECK
-	CHECK(A.NumNestedScopes == B.NumNestedScopes);
-	CHECK(A.NestedScopesOffset == B.NestedScopesOffset);
-	CHECK(A.NumParametricTypes == B.NumParametricTypes);
-	CHECK(A.NumSchemas == B.NumSchemas);
-	CHECK(A.NumStructSchemas == B.NumStructSchemas);
-	CHECK(EqualItems(A.GetSchemaOffsets(), B.GetSchemaOffsets()));
-	CHECK(EqualItems(A.GetNestedScopes(), B.GetNestedScopes()));
-	CHECK(EqualItems(A.GetParametricTypes(), B.GetParametricTypes()));
-
-	{
-		TSchemaRange<const FStructSchema> AA = GetStructSchemas(A);
-		TSchemaRange<const FStructSchema> BB = GetStructSchemas(B);
-		TSchemaIterator<const FStructSchema> ItA = AA.begin();
-		TSchemaIterator<const FStructSchema> ItB = BB.begin();
-		for (; ItA != AA.end() && ItB != BB.end(); ++ItA, ++ItB)
-		{
-			CHECK_FALSE(DiffStructSchema(*ItA, *ItB));
-		}
-	}
-
-	{
-		TSchemaRange<const FEnumSchema> AA = GetEnumSchemas(A);
-		TSchemaRange<const FEnumSchema> BB = GetEnumSchemas(B);
-		TSchemaIterator<const FEnumSchema> ItA = AA.begin();
-		TSchemaIterator<const FEnumSchema> ItB = BB.begin();
-		for (; ItA != AA.end() && ItB != BB.end(); ++ItA, ++ItB)
-		{
-			CHECK_FALSE(DiffEnumSchema(*ItA, *ItB));
-		}
-	}
-
-	CHECK(A == B);
-#endif
-	return Diff;
-}
-
-bool DiffSchemas(FMemoryView A, FMemoryView B)
-{
-	const bool BinaryDiff = A.GetSize() != B.GetSize() || FMemory::Memcmp(A.GetData(), B.GetData(), A.GetSize());
-#if DO_CHECK
-	const FSchemaBatch* SchemasA = reinterpret_cast<const FSchemaBatch*>(A.GetData());
-	const FSchemaBatch* SchemasB = reinterpret_cast<const FSchemaBatch*>(B.GetData());
-	SchemasA->ValidateBounds(A.GetSize());
-	SchemasB->ValidateBounds(B.GetSize());
-	const bool Diff = DiffSchemas(*SchemasA, *SchemasB);
-	checkf(Diff == BinaryDiff, TEXT("DiffSchemas() should produce the same result as a binary diff"));
-#endif
-	return BinaryDiff;
-}
-
 ////////////////////////////////////////////////////////////////////////////////
 
 TEST_CASE_NAMED(FPlainPropsIndexTest, "System::Core::Serialization::PlainProps::Index", "[Core][PlainProps][SmokeFilter]")
@@ -275,7 +136,6 @@ public:
 	const FStructDeclaration&	Get(FStructId Id) { return Declarations.Get(Id); }
 
 	void						AddObject(FStructId Schema, FMemberBuilder&& Members);
-	void						AddObjects(FStructId Schema, std::initializer_list<FMemberBuilder&&> Objects);
 
 	TArray64<uint8>				Write();
 
@@ -456,28 +316,68 @@ public:
 	virtual void AppendString(FUtf8StringBuilderBase& Out, FNameId Name) const override { Out.Append(Names[Name]); }
 };
 
-static void PrintBatch(
-	FUtf8StringBuilderBase& Out,
-	const FTestNameReader& Names,
-	FReadBatchId Batch,
-	const FSchemaBatch& Schemas,
-	TConstArrayView<FStructView> Objects)
+[[nodiscard]] FReadBatchId ParseBatch(TArray64<uint8>& OutData, TArray<FStructView>& OutObjects, FUtf8StringView YamlView)
 {
-	FTestBatchNames BatchNames(Names, Batch);
-	PrintYamlBatch(Out, BatchNames, Objects);
-	FUtf8StringView YamlView = Out.ToView();
-	auto Wide = StringCast<TCHAR>(YamlView.GetData(), YamlView.Len());
-	UE_LOG(LogPlainPropsTests, Log, TEXT("\n%.*s\n"), Wide.Length(), Wide.Get());
+	// Parse yaml
+	ParseYamlBatch(OutData, YamlView);
+
+	// Grab and mount parsed schemas
+	FByteReader It(MakeMemoryView(OutData));
+	const uint32 SchemasSize = It.Grab<uint32>();
+	FMemoryView SchemasView = It.GrabSlice(SchemasSize);
+	const FSchemaBatch* Schemas = ValidateSchemas(SchemasView);
+	FReadBatchId Batch = MountReadSchemas(Schemas);
+
+	// Grab parsed objects
+	while (uint64 NumBytes = It.GrabVarIntU())
+	{	
+		FByteReader ObjIt(It.GrabSlice(NumBytes));
+		FStructSchemaId Schema = { ObjIt.Grab<uint32>() };
+		OutObjects.Add({ { Schema, Batch }, ObjIt });
+	}
+	
+	return Batch;
 }
 
-static void ParseBatch(FUtf8StringView YamlView, FMemoryView BinaryView)
+static void RoundtripText(const FTestNameReader& Names, FReadBatchId Batch, TConstArrayView<FStructView> Objects)
 {
+	// Print yaml
+	TUtf8StringBuilder<4096> Yaml;
+	FTestBatchNames BatchNames(Names, Batch);
+	PrintYamlBatch(Yaml, BatchNames, Objects);
+	FUtf8StringView YamlView = Yaml.ToView();
+
+	// Log yaml
+	auto Wide = StringCast<TCHAR>(YamlView.GetData(), YamlView.Len());
+	UE_LOG(LogPlainPropsTests, Log, TEXT("Schemas with StableNames:\n%.*s"), Wide.Length(), Wide.Get());
+
+	// Parse yaml
 	TArray64<uint8> Data;
-	ParseYamlBatch(Data, YamlView);
-	FMemoryView ParsedBinaryView(Data.GetData(), Data.Num());
-	CHECK_FALSE(DiffSchemas(BinaryView, ParsedBinaryView));
-}
+	TArray<FStructView> ParsedObjects;
+	FReadBatchId ParsedBatch = ParseBatch(Data, ParsedObjects, YamlView);
+
+	// Diff schemas
+	CHECK_FALSE(DiffSchemas(Batch, ParsedBatch));
+
+	// Diff objects
+	CHECK(Objects.Num() == ParsedObjects.Num());
+	const int32 NumObjects = FMath::Min(Objects.Num(), ParsedObjects.Num());
+	for (int32 I = 0; I < NumObjects; ++I)
+	{
+		FStructView In = Objects[I];
+		FStructView Parsed = ParsedObjects[I];
+		FReadDiffPath DiffPath;
+		if (DiffStruct(In, Parsed, DiffPath))
+		{
+			TUtf8StringBuilder<256> Diff;
+			PrintDiff(Diff, BatchNames, DiffPath);
+			FAIL_CHECK(FString::Printf(TEXT("Diff in '%s' in Objects[%d]"), *Print(Diff.ToString()), I));
+		}
+	}
 
+	// Unmount parsed schemas
+	UnmountReadSchemas(ParsedBatch);
+}
 ////////////////////////////////////////////////////////////////////////////////////////////////
 
 class FTestBatchReader
@@ -510,11 +410,11 @@ public:
 		
 		CHECK(It.Grab<uint32>() == TestMagics[4]);
 		CHECK(!Objects.IsEmpty());
+	}
 
-		// Roundtrip schemas and objects to text
-		TUtf8StringBuilder<4096> YamlString;
-		PrintBatch(YamlString, Names, Batch, *Schemas, Objects);
-		ParseBatch(YamlString.ToView(), SchemasView);
+	void RoundtripText()
+	{
+		PlainProps::RoundtripText(Names, Objects[0].Schema.Batch, Objects);
 	}
 
 	~FTestBatchReader()
@@ -541,6 +441,7 @@ static void TestSerialize(void (*BuildObjects)(FTestBatchBuilder&, FScratchAlloc
 	}
 
 	FTestBatchReader Batch(MakeMemoryView(Data));
+	Batch.RoundtripText();
 	CheckObjects(Batch.GetObjects(), Batch.GetNames());
 }
 
@@ -776,6 +677,66 @@ TEST_CASE_NAMED(FPlainPropsReadWriteTest, "System::Core::Serialization::PlainPro
 		});
 	}
 
+	SECTION("Unicode")
+	{
+		TestSerialize([](FTestBatchBuilder& Batch, FScratchAllocator& Scratch)
+		{
+			FStructId Char8Id  = Batch.DeclareStruct("Test", "Char8",  {"A", "B", "C", "D", "E", "F"}, EMemberPresence::AllowSparse);
+			FStructId Char16Id = Batch.DeclareStruct("Test", "Char16", {"A", "B", "C", "D", "E", "F"}, EMemberPresence::AllowSparse);
+			FStructId Char32Id = Batch.DeclareStruct("Test", "Char32", {"A", "B", "C", "D", "E", "F"}, EMemberPresence::AllowSparse);
+
+			FMemberBuilder Char8, Char16, Char32;
+
+			Char8.Add(Batch.NameMember("A"), u8'\0');		// NUL, first valid code unit
+			Char8.Add(Batch.NameMember("B"), u8'\x1');		// SOH, a control character
+			Char8.Add(Batch.NameMember("C"), u8'\n');		// LF, an escaped character
+			Char8.Add(Batch.NameMember("D"), u8'%');		// %, a printable character
+			Char8.Add(Batch.NameMember("E"), u8'E');		// E, an alphabetic character
+			Char8.Add(Batch.NameMember("F"), u8'\x7F');		// DEL, last valid code unit
+
+			Char16.Add(Batch.NameMember("A"), u'\0');		// First valid code unit
+			Char16.Add(Batch.NameMember("B"), u'\x0024');	// Dollar sign, single byte code unit
+			Char16.Add(Batch.NameMember("C"), u'\xD7FF');	// Last single byte code unit
+			Char16.Add(Batch.NameMember("D"), u'\xE000');	// First double byte code unit
+			Char16.Add(Batch.NameMember("E"), u'\x20AC');	// Euro sign, double byte code unit
+			Char16.Add(Batch.NameMember("F"), u'\xFFFD');	// Last valid single code unit character
+			// UE::Core::Private::IsValidCodepoint() treats the non-characters FFFE and FFFF as invalid
+						
+			Char32.Add(Batch.NameMember("A"), U'\0');
+			Char32.Add(Batch.NameMember("B"), U'\x1');
+			Char32.Add(Batch.NameMember("C"), U'C');
+
+			Batch.AddObject(Char8Id, MoveTemp(Char8));
+			Batch.AddObject(Char16Id, MoveTemp(Char16));
+			Batch.AddObject(Char32Id, MoveTemp(Char32));
+		}, 
+		[](TConstArrayView<FStructView> Objects, const FTestNameReader& Names)
+		{
+			CHECK(Objects.Num() == 3);
+			
+			FTestMemberReader Char8(Objects[0]);
+			CHECK(Char8.GrabLeaf().AsChar8() == u8'\0');
+			CHECK(Char8.GrabLeaf().AsChar8() == u8'\x1');
+			CHECK(Char8.GrabLeaf().AsChar8() == u8'\n');
+			CHECK(Char8.GrabLeaf().AsChar8() == u8'%');
+			CHECK(Char8.GrabLeaf().AsChar8() == u8'E');
+			CHECK(Char8.GrabLeaf().AsChar8() == u8'\x7F');
+
+			FTestMemberReader Char16(Objects[1]);
+			CHECK(Char16.GrabLeaf().AsChar16() == u'\0');
+			CHECK(Char16.GrabLeaf().AsChar16() == u'\x0024');
+			CHECK(Char16.GrabLeaf().AsChar16() == u'\xD7FF');
+			CHECK(Char16.GrabLeaf().AsChar16() == u'\xE000');
+			CHECK(Char16.GrabLeaf().AsChar16() == u'\x20AC');
+			CHECK(Char16.GrabLeaf().AsChar16() == u'\xFFFD');
+
+			FTestMemberReader Char32(Objects[2]);
+			CHECK(Char32.GrabLeaf().AsChar32() == U'\0');
+			CHECK(Char32.GrabLeaf().AsChar32() == U'\x1');
+			CHECK(Char32.GrabLeaf().AsChar32() == U'C');
+		});
+	}
+
 	SECTION("Dense")
 	{
 		static constexpr std::initializer_list<const char*> ExplicitMembers = {"A", "B", "C"};
@@ -1081,6 +1042,54 @@ TEST_CASE_NAMED(FPlainPropsReadWriteTest, "System::Core::Serialization::PlainPro
 		});
 	}
 
+	SECTION("UnicodeRange")
+	{
+		TestSerialize([](FTestBatchBuilder& Batch, FScratchAllocator& Scratch)
+		{
+			FStructId Utf8Id  = Batch.DeclareStruct("Test", "Utf8",  {"Null", "Empty", "Escape", "Latin1", "CJK", "Symbols"}, EMemberPresence::AllowSparse);
+			FStructId Utf16Id = Batch.DeclareStruct("Test", "Utf16", {"Null", "Empty", "Escape", "Latin1", "CJK", "Symbols"}, EMemberPresence::AllowSparse);
+
+			FMemberBuilder Utf8, Utf16;
+
+			Utf8.AddRange(Batch.NameMember("Null"),		BuildLeafRange(Scratch, TConstArrayView<char8_t>()));
+			Utf8.AddRange(Batch.NameMember("Empty"),	BuildLeafRange(Scratch, MakeArrayView(u8"")));
+			Utf8.AddRange(Batch.NameMember("Escape"),	BuildLeafRange(Scratch, MakeArrayView(u8"\"\\ \x01 \x1f\" \"\b \f \n \r \t \"\\").LeftChop(1)));
+			Utf8.AddRange(Batch.NameMember("Latin1"),	BuildLeafRange(Scratch, MakeArrayView(u8"\xE5 \xE4 \xF6")));		// å ä ö
+			Utf8.AddRange(Batch.NameMember("CJK"),		BuildLeafRange(Scratch, MakeArrayView(u8"\u3300 \uFE30")));			// ㌀ ︰
+			Utf8.AddRange(Batch.NameMember("Symbols"),	BuildLeafRange(Scratch, MakeArrayView(u8"\u2665 \u01F34C")));		// ♥ 🍌
+
+			Utf16.AddRange(Batch.NameMember("Null"),	BuildLeafRange(Scratch, TConstArrayView<char16_t>()));
+			Utf16.AddRange(Batch.NameMember("Empty"),	BuildLeafRange(Scratch, MakeArrayView(u"")));
+			Utf16.AddRange(Batch.NameMember("Escape"),	BuildLeafRange(Scratch, MakeArrayView(u"\\\" \x01 \x1f\" \"\b \f \n \r \t \\\"").LeftChop(1)));
+			Utf16.AddRange(Batch.NameMember("Latin1"),	BuildLeafRange(Scratch, MakeArrayView(u"\xC5 \xC4 \xD6")));			// Å Ä Ö
+			Utf16.AddRange(Batch.NameMember("CJK"),		BuildLeafRange(Scratch, MakeArrayView(u"\x3300 \xFE30")));			// ㌀ ︰
+			Utf16.AddRange(Batch.NameMember("Symbols"),	BuildLeafRange(Scratch, MakeArrayView(u"\x2665 \xD83C\xDF4C")));	// ♥ 🍌
+
+			Batch.AddObject(Utf8Id, MoveTemp(Utf8));
+			Batch.AddObject(Utf16Id, MoveTemp(Utf16));
+		}, 
+		[](TConstArrayView<FStructView> Objects, const FTestNameReader& Names)
+		{
+			CHECK(Objects.Num() == 2);
+
+			FTestMemberReader It1(Objects[0]);
+			CHECK(EqualItems(It1.GrabRange().AsLeaves().AsUtf8(), TConstArrayView<char8_t>()));
+			CHECK(EqualItems(It1.GrabRange().AsLeaves().AsUtf8(), u8""));
+			CHECK(EqualItems(It1.GrabRange().AsLeaves().AsUtf8(), MakeArrayView(u8"\"\\ \x01 \x1f\" \"\b \f \n \r \t \"\\").LeftChop(1)));
+			CHECK(EqualItems(It1.GrabRange().AsLeaves().AsUtf8(), u8"\xE5 \xE4 \xF6"));			// å ä ö
+			CHECK(EqualItems(It1.GrabRange().AsLeaves().AsUtf8(), u8"\u3300 \uFE30"));			// ㌀ ︰
+			CHECK(EqualItems(It1.GrabRange().AsLeaves().AsUtf8(), u8"\u2665 \u01F34C"));		// ♥ 🍌
+
+			FTestMemberReader It2(Objects[1]);
+			CHECK(EqualItems(It2.GrabRange().AsLeaves().AsUtf16(), TConstArrayView<char16_t>()));
+			CHECK(EqualItems(It2.GrabRange().AsLeaves().AsUtf16(), u""));
+			CHECK(EqualItems(It2.GrabRange().AsLeaves().AsUtf16(), MakeArrayView(u"\\\" \x01 \x1f\" \"\b \f \n \r \t \\\"").LeftChop(1)));
+			CHECK(EqualItems(It2.GrabRange().AsLeaves().AsUtf16(), u"\xC5 \xC4 \xD6"));			// Å Ä Ö
+			CHECK(EqualItems(It2.GrabRange().AsLeaves().AsUtf16(), u"\x3300 \xFE30"));			// ㌀ ︰
+			CHECK(EqualItems(It2.GrabRange().AsLeaves().AsUtf16(), u"\x2665 \xD83C\xDF4C"));	// ♥ 🍌
+		});
+	}
+		
 	SECTION("StructRange")
 	{
 		TestSerialize([](FTestBatchBuilder& Batch, FScratchAllocator& Scratch)
@@ -1378,7 +1387,8 @@ TEST_CASE_NAMED(FPlainPropsReadWriteTest, "System::Core::Serialization::PlainPro
 			O2.AddRange(Batch.NameMember("Same"), BuildStructRange(SA, 20));
 			O2.AddRange(Batch.NameMember("Diff"), BuildStructRange(SB, 22.f));
 			O2.AddRange(Batch.NameMember("SameEmpty"), FStructRangeBuilder(0).BuildAndReset(Scratch, Batch.Get(SA), Batch.GetDebug()));
-			O2.AddRange(Batch.NameMember("DiffEmpty"), FStructRangeBuilder(0).BuildAndReset(Scratch, Batch.Get(SB), Batch.GetDebug()));
+			// PP-TEXT: Handle DiffEmpty in DynamicStructRange Test by printing dynamic type for empty ranges
+			// O2.AddRange(Batch.NameMember("DiffEmpty"), FStructRangeBuilder(0).BuildAndReset(Scratch, Batch.Get(SB), Batch.GetDebug()));
 			FNestedRangeBuilder NestedSB(MakeStructRangeSchema(ERangeSizeType::S32, SB), 1);
 			NestedSB.Add(BuildStructRange(SB, 200.f));
 			O2.AddRange(Batch.NameMember("DiffNested"), NestedSB.BuildAndReset(Scratch, ERangeSizeType::S32));
@@ -1403,7 +1413,7 @@ TEST_CASE_NAMED(FPlainPropsReadWriteTest, "System::Core::Serialization::PlainPro
 			CHECK(MakeArray<FTestMemberReader>(O2.GrabRange().AsStructs())[0].GrabLeaf().AsS32() == 20);
 			CHECK(MakeArray<FTestMemberReader>(O2.GrabRange().AsStructs())[0].GrabLeaf().AsFloat() == 22.f);
 			CHECK(O2.GrabRange().AsStructs().Num() == 0);
-			CHECK(O2.GrabRange().AsStructs().Num() == 0);
+			// CHECK(O2.GrabRange().AsStructs().Num() == 0);
 			TArray<FRangeView> DiffNested2 = MakeArray<FRangeView>(O2.GrabRange().AsRanges());
 			CHECK(MakeArray<FTestMemberReader>(DiffNested2[0].AsStructs())[0].GrabLeaf().AsFloat() == 200.f);
 		});
diff --git a/Engine/Plugins/Experimental/PlainProps/Source/Tests/PlainPropsUeCoreTests.cpp b/Engine/Plugins/Experimental/PlainProps/Source/Tests/PlainPropsUeCoreTests.cpp
--- a/Engine/Plugins/Experimental/PlainProps/Source/Tests/PlainPropsUeCoreTests.cpp
+++ b/Engine/Plugins/Experimental/PlainProps/Source/Tests/PlainPropsUeCoreTests.cpp
@@ -6,11 +6,13 @@
 #include "PlainPropsCtti.h"
 #include "PlainPropsIndex.h"
 #include "PlainPropsInternalBuild.h"
+#include "PlainPropsInternalDiff.h"
 #include "PlainPropsInternalFormat.h"
 #include "PlainPropsInternalParse.h"
 #include "PlainPropsInternalPrint.h"
 #include "PlainPropsInternalRead.h"
 #include "PlainPropsInternalTest.h"
+#include "PlainPropsDiff.h"
 #include "PlainPropsLoad.h"
 #include "PlainPropsRead.h"
 #include "PlainPropsSave.h"
@@ -325,48 +327,69 @@ public:
 	virtual FParametricTypeView	Resolve(FParametricTypeId Id) const override		{ return GNames.Resolve(Binding.Remap(Id)); }
 };
 
-static void PrintBatch(FUtf8Builder& Out, TConstArrayView<FName> Names,	const FSchemaBatch& Schemas, TArray<FStructView>&& Objects)
+static void RoundtripText(TConstArrayView<FName> Names, const FSchemaBatch& Schemas, TArray<FStructView>&& Objects)
 {
-	FReadBatchId Batch = MountReadSchemas(&Schemas);
+	ESchemaFormat Format = Names.IsEmpty() ? ESchemaFormat::InMemoryNames : ESchemaFormat::StableNames;
 
+	// Mount input schemas and update copy of objects array accordingly
+	FReadBatchId Batch = MountReadSchemas(&Schemas);
 	for (FStructView& Struct : Objects)
 	{
 		Struct.Schema.Batch = Batch;
 	}
-	
-	if (Names.IsEmpty())
+
+	// Print yaml
+	TUtf8StringBuilder<4096> Yaml;
+	if (Format == ESchemaFormat::InMemoryNames)
 	{
 		FMemoryBatchIds Ids(Batch, GNames);
-		PrintYamlBatch(Out, Ids, Objects);
+		PrintYamlBatch(Yaml, Ids, Objects);
 	}
 	else
 	{
 		FStableNameBatchIds Ids(Batch, Names);
-		PrintYamlBatch(Out, Ids, Objects);
+		PrintYamlBatch(Yaml, Ids, Objects);
 	}
-	
-	UnmountReadSchemas(Batch);
+	FUtf8StringView YamlView = Yaml.ToView();
 
-	FUtf8StringView YamlView = Out.ToView();
+	// Log yaml
 	auto Wide = StringCast<TCHAR>(YamlView.GetData(), YamlView.Len());
-	UE_LOG(LogPlainPropsUeCoreTests, Log, TEXT("%s Schemas:\n%.*s"),
-		Names.IsEmpty() ? TEXT("InMemory") : TEXT("StableNames"),
+	UE_LOG(LogPlainPropsUeCoreTests, Log, TEXT("Schemas with %s:\n%.*s"),
+		Format == ESchemaFormat::InMemoryNames ? TEXT("InMemoryNames") : TEXT("StableNames"),
 		Wide.Length(), Wide.Get());
-}
 
-static void ParseBatch(FUtf8StringView YamlView, FMemoryView BinaryView, ESchemaFormat Format)
-{
+	// Parse yaml
 	TArray64<uint8> Data;
+	TArray<FStructView> ParsedObjects;
+	FReadBatchId ParsedBatch = ParseBatch(Data, ParsedObjects, YamlView);
+
 	if (Format == ESchemaFormat::StableNames)
 	{
-		ParseYamlBatch(Data, YamlView);
-		FMemoryView ParsedBinaryView(Data.GetData(), Data.Num());
-		CHECK_FALSE(DiffSchemas(BinaryView, ParsedBinaryView));
-	}
-	else
-	{
-		ParseYamlBatch(Data, YamlView);
+		FStableNameBatchIds Ids(Batch, Names);
+
+		// Diff schemas
+		CHECK_FALSE(DiffSchemas(Batch, ParsedBatch));
+
+		// Diff objects
+		CHECK(Objects.Num() == ParsedObjects.Num());
+		const int32 NumObjects = FMath::Min(Objects.Num(), ParsedObjects.Num());
+		for (int32 I = 0; I < NumObjects; ++I)
+		{
+			FStructView In = Objects[I];
+			FStructView Parsed = ParsedObjects[I];
+			FReadDiffPath DiffPath;
+			if (DiffStruct(In, Parsed, DiffPath))
+			{
+				TUtf8StringBuilder<256> Diff;
+				PrintDiff(Diff, Ids, DiffPath);
+				FAIL_CHECK(FString::Printf(TEXT("Diff in '%s' in Objects[%d]"), *Print(Diff.ToString()), I));
+			}
+		}
 	}
+
+	// Unmount input schemas and parsed schemas
+	UnmountReadSchemas(Batch);
+	UnmountReadSchemas(ParsedBatch);
 }
 
 static void PrintLoadStruct(const FIdBinding& Binding, FStructView StructView, int32 LoadIdx)
@@ -403,8 +426,8 @@ public:
 		CHECK(It.Grab<uint32>() == Magics[1]);
 		It.SkipAlignmentPadding<uint32>();
 		uint32 SchemasSize = It.Grab<uint32>();
-		FMemoryView SavedSchemasView = It.GrabSlice(SchemasSize);
-		const FSchemaBatch* SavedSchemas = ValidateSchemas(SavedSchemasView);
+		FMemoryView	SavedSchemasView = It.GrabSlice(SchemasSize);
+		SavedSchemas = ValidateSchemas(SavedSchemasView);
 		CHECK(It.Grab<uint32>() == Magics[2]);
 
 
@@ -450,10 +473,6 @@ public:
 		TConstArrayView<FStructId> LoadStructIds = RuntimeIds ? RuntimeIds->Translation.GetStructIds(SavedSchemas->NumStructSchemas) : InRuntimeIds;
 		Plans = CreateLoadPlans(Batch, GTypes, Customs, GSchemas, LoadStructIds, Format);
 
-		// Mount, print, unmount, and then parse the saved batch
-		TUtf8StringBuilder<4096> YamlBatch;
-		PrintBatch(YamlBatch, Ids, *SavedSchemas, TArray<FStructView>(Objects));
-		ParseBatch(YamlBatch.ToView(), SavedSchemasView, Format);
 	}
 
 	~FBatchLoader()
@@ -467,6 +486,11 @@ public:
 		}
 	}
 
+	void RoundtripText() const
+	{
+		PlainProps::UE::Test::RoundtripText(Ids, *SavedSchemas, CopyTemp(Objects));
+	}
+
 	template<class T>
 	T Load()
 	{
@@ -496,6 +520,7 @@ public:
 	}
 
 private:
+	const FSchemaBatch*			SavedSchemas;
 	TConstArrayView<FName>		Ids;
 	FNameBinding				Names;
 	TOptional<FIdBinding>		IdBinding;
@@ -519,6 +544,7 @@ static void Run(void (*Save)(FBatchSaver&), void (*Load)(FBatchLoader&), const F
 		}
 
 		FBatchLoader Batch(MakeMemoryView(Data), Customs, MakeArrayView(RuntimeIds));
+		Batch.RoundtripText();
 		Load(Batch);
 	}
 }
diff --git a/Engine/Plugins/Experimental/PlainPropsUObject/Source/Private/PlainPropsUObjectBindings.cpp b/Engine/Plugins/Experimental/PlainPropsUObject/Source/Private/PlainPropsUObjectBindings.cpp
--- a/Engine/Plugins/Experimental/PlainPropsUObject/Source/Private/PlainPropsUObjectBindings.cpp
+++ b/Engine/Plugins/Experimental/PlainPropsUObject/Source/Private/PlainPropsUObjectBindings.cpp
@@ -4,11 +4,14 @@
 #include "PlainPropsUeCoreBindings.h"
 #include "PlainPropsBuildSchema.h"
 #include "PlainPropsDiff.h"
+#include "PlainPropsParse.h"
 #include "PlainPropsPrint.h"
+#include "PlainPropsRead.h"
 #include "PlainPropsWrite.h"
 #include "Algo/Compare.h"
 #include "Algo/Find.h"
 #include "Containers/PagedArray.h"
+#include "HAL/FileManager.h"
 #include "Hash/xxhash.h"
 #include "Logging/StructuredLog.h"
 #include "Math/PreciseFP.h"
@@ -16,9 +19,12 @@
 #include "Misc/CommandLine.h"
 #include "Misc/CoreDelegates.h"
 #include "Misc/DefinePrivateMemberPtr.h"
+#include "Misc/Paths.h"
+#include "Serialization/Archive.h"
 #include "Serialization/MemoryReader.h"
 #include "Serialization/MemoryWriter.h"
 #include "StructUtils/UserDefinedStruct.h"
+#include "Templates/UniquePtr.h"
 #include "UObject/AnsiStrProperty.h"
 #include "UObject/Class.h"
 #include "UObject/EnumProperty.h"
@@ -1012,6 +1018,17 @@ static bool DenyMaxValue(FName Enum)
 	return !Algo::Find(AllowsMax, Enum);
 }
 
+enum class ERoundtrip : uint8
+{
+	None 		= 0,
+	PP 			= 1 << 0,
+	TPS 		= 1 << 1,
+	UPS 		= 1 << 2,
+	TextMemory	= 1 << 3,
+	TextStable	= 1 << 4,
+};
+ENUM_CLASS_FLAGS(ERoundtrip);
+
 static FEnumId DeclareEnum(UEnum* Enum)
 {
 	FType Type = IndexType(Enum);
@@ -3591,11 +3608,13 @@ public:
 		SavedObjects.Emplace(Id, Built, Object);
 	}
 
-	TArray64<uint8>	Write(TArray<FStructId>& OutRuntimeIds) const
+	TArray64<uint8>	Write(TArray<FStructId>* OutRuntimeIds = nullptr) const
 	{
+		ESchemaFormat Format = OutRuntimeIds ? ESchemaFormat::InMemoryNames : ESchemaFormat::StableNames;
+
 		// Build partial schemas
 		const FStructBindIds BindIds(GUE.Customs, GUE.Schemas);
-		FSchemasBuilder SchemaBuilders(GUE.Types, GUE.Names, BindIds, Scratch, ESchemaFormat::InMemoryNames);
+		FSchemasBuilder SchemaBuilders(GUE.Types, GUE.Names, BindIds, Scratch, Format);
 		for (const FSavedObject& Object : SavedObjects)
 		{
 			GDebugNoteObject = Object.Input;
@@ -3604,12 +3623,30 @@ public:
 		GDebugNoteObject = nullptr;
 		FBuiltSchemas Schemas = SchemaBuilders.Build();
 
-		// Save schema ids on the side
-		OutRuntimeIds = ExtractRuntimeIds(Schemas);
+		// Save schema ids on the side when using InMemoryNames
+		if (OutRuntimeIds)
+		{
+			*OutRuntimeIds = ExtractRuntimeIds(Schemas);
+		}
 
-		// Write schemas
-		FWriter Writer(GUE.Names, BindIds, Schemas, ESchemaFormat::InMemoryNames);
+		FWriter Writer(GUE.Names, BindIds, Schemas, Format);
 		TArray64<uint8> Out;
+
+		// Write out FNames when using StableNames
+		if (!OutRuntimeIds)
+		{
+			TArray<FName> UsedNames;
+			UsedNames.Reserve(Writer.GetUsedNames().Num());
+			for (FNameId Name : Writer.GetUsedNames())
+			{
+				UsedNames.Add(GUE.Names.ResolveName(Name));
+			}
+
+			WriteInt(Out, Magics[0]);
+			WriteNumAndArray(Out, TArrayView<const FName, int32>(UsedNames));
+		}
+
+		// Write schemas
 		WriteInt(Out, Magics[1]);
 		WriteAlignmentPadding<uint32>(Out);
 		TArray64<uint8> Tmp;
@@ -3674,15 +3711,12 @@ public:
 		FReadBatchId Batch = MountReadSchemas(SavedSchemas);
 
 		// Read objects
-		TSet<FStructSchemaId> SchemaIds;
-		SchemaIds.Reserve(NumStructSchemas(Batch));
  		while (uint64 NumBytes = It.GrabVarIntU())
 		{	
 			FByteReader ObjIt(It.GrabSlice(NumBytes));
 			check(ObjIt.Grab<uint32>() == Magics[3]);
 			FStructSchemaId Id = { ObjIt.Grab<uint32>() };
 			Objects.Add({ { Id, Batch }, ObjIt });
-			SchemaIds.Add(Id);
 		}
 		
 		check(It.Grab<uint32>() == Magics[4]);
@@ -3826,6 +3860,7 @@ public:
 
 struct FInstance
 {
+	FString PathName;
 	UObject* Orig;
 	UObject* Arch;
 	UObject* Base;
@@ -3924,6 +3959,184 @@ FORCENOINLINE void LoadArchive(FPropertyReader& Archive, TConstArrayView<FInstan
 	}
 }
 
+////////////////////////////////////////////////////////////////////////////////////////////////
+
+class FStableNameBatchIds final : public FStableBatchIds
+{
+	TConstArrayView<FName> Names;
+public:
+	FStableNameBatchIds(FReadBatchId Batch, TConstArrayView<FName> InNames) : FStableBatchIds(Batch), Names(InNames) {}
+	using FStableBatchIds::AppendString;
+
+	virtual uint32				NumNames() const override							{ return static_cast<uint32>(Names.Num()); }
+	virtual void				AppendString(FUtf8Builder& Out, FNameId Name) const override { Names[Name.Idx].AppendString(Out); }
+};
+
+[[nodiscard]] static FReadBatchId ParseBatch(
+	TArray64<uint8>& OutData,
+	TArray<FStructView>& OutObjects,
+	FUtf8StringView YamlView)
+{
+	// Parse yaml
+	ParseYamlBatch(OutData, YamlView);
+
+	// Grab and mount parsed schemas
+	FByteReader It(MakeMemoryView(OutData));
+	const uint32 SchemasSize = It.Grab<uint32>();
+	FMemoryView SchemasView = It.GrabSlice(SchemasSize);
+	const FSchemaBatch* Schemas = ValidateSchemas(SchemasView);
+	FReadBatchId Batch = MountReadSchemas(Schemas);
+
+	// Grab parsed objects
+	while (uint64 NumBytes = It.GrabVarIntU())
+	{	
+		FByteReader ObjIt(It.GrabSlice(NumBytes));
+		FStructSchemaId Schema = { ObjIt.Grab<uint32>() };
+		OutObjects.Add({ { Schema, Batch }, ObjIt });
+	}
+	
+	return Batch;
+}
+
+static void RoundtripText(
+	TConstArrayView<FName> Names,
+	FMemoryView SchemasView,
+	FReadBatchId Batch,
+	TConstArrayView<FStructView> Objects,
+	TConstArrayView<FInstance> Instances)
+{
+	check(Objects.Num() == Instances.Num());
+	ESchemaFormat Format = Names.IsEmpty() ? ESchemaFormat::InMemoryNames : ESchemaFormat::StableNames;
+
+	// Print yaml
+	UE_LOGFMT(LogPlainPropsUObject, Display, "Printing to PlainProps text using {Format}...", ToString(Format));
+	TUtf8StringBuilder<256> Yaml;
+	Yaml.Reserve((2 << 30) - 1);
+	if (Format == ESchemaFormat::InMemoryNames)
+	{
+		FMemoryBatchIds Ids(Batch, GUE.Names);
+		PrintYamlBatch(Yaml, Ids, Objects);
+	}
+	else
+	{
+		FStableNameBatchIds Ids(Batch, Names);
+		PrintYamlBatch(Yaml, Ids, Objects);
+	}
+	FUtf8StringView YamlView = Yaml.ToView();
+
+	// Write to file
+	const FString Filename = FPaths::ProjectSavedDir() / TEXT("PlainProps") /
+		(Format == ESchemaFormat::InMemoryNames ? TEXT("InMemoryNames.yaml") : TEXT("StableNames.yaml"));
+	if (TUniquePtr<FArchive> FileWriter(IFileManager::Get().CreateFileWriter(*Filename)); FileWriter)
+	{
+		UE_LOGFMT(LogPlainPropsUObject, Display, "Writing {KB}KB yaml as {Filename}...", Yaml.Len() >> 10, *Filename);
+		FileWriter->Serialize((void*)YamlView.GetData(), YamlView.Len());
+	}
+
+	// Parse yaml
+	UE_LOGFMT(LogPlainPropsUObject, Display, "Parsing PlainProps text using {Format}...", ToString(Format));
+	TArray64<uint8> Data;
+	TArray<FStructView> ParsedObjects;
+	FReadBatchId ParsedBatch = ParseBatch(Data, ParsedObjects, YamlView);
+
+	if (Format == ESchemaFormat::StableNames)
+	{
+		UE_LOGFMT(LogPlainPropsUObject, Display, "Diffing PlainProps parsed objects using {Format}...", ToString(Format));
+		FStableNameBatchIds Ids(Batch, Names);
+
+		// Diff schemas
+		check(!DiffSchemas(Batch, ParsedBatch));
+
+		// Diff objects
+		check(Objects.Num() == ParsedObjects.Num());
+		const int32 NumObjects = FMath::Min(Objects.Num(), ParsedObjects.Num());
+		uint32 NumDiffs = 0;
+		TUtf8StringBuilder<256> Diffs;
+		for (int32 I = 0; I < NumObjects; ++I)
+		{
+			FStructView In = Objects[I];
+			FStructView Parsed = ParsedObjects[I];
+			FReadDiffPath DiffPath;
+			if (DiffStruct(In, Parsed, DiffPath))
+			{
+				PrintDiff(Diffs, Ids, DiffPath);
+				Diffs.Append(" in ");
+				Diffs.Append(Instances[I].PathName);
+				Diffs.Append("\n");
+				++NumDiffs;
+			}
+		}
+		UE_LOGFMT(LogPlainPropsUObject, Display,
+			"Detected {Diffs} diffs in {Objs} PlainProps parsed objects from {KB}KB yaml text using StableNames\n{Diffs}",
+			NumDiffs, NumObjects, Yaml.Len() >> 10, Diffs.ToString());
+	}
+
+	// Unmount parsed schemas
+	UnmountReadSchemas(ParsedBatch);
+}
+
+class FBatchTextRoundtripper
+{
+public:
+	FBatchTextRoundtripper(FMemoryView Data, ESchemaFormat Format)
+	{
+		FByteReader It(Data);
+
+		// Read FNames when using Stable Names
+		if (Format == ESchemaFormat::StableNames)
+		{
+			verify(It.Grab<uint32>() == Magics[0]);
+			Ids = GrabNumAndArray<FName>(It);
+		}
+		
+		// Read and mount schemas
+		check(It.Grab<uint32>() == Magics[1]);
+		It.SkipAlignmentPadding<uint32>();
+		const uint32 SchemasSize = It.Grab<uint32>();
+		SavedSchemasView = It.GrabSlice(SchemasSize);
+		const FSchemaBatch* SavedSchemas = ValidateSchemas(SavedSchemasView);
+		check(It.Grab<uint32>() == Magics[2]);
+		Batch = MountReadSchemas(SavedSchemas);
+		
+		// Read objects
+		while (uint64 NumBytes = It.GrabVarIntU())
+		{	
+			FByteReader ObjIt(It.GrabSlice(NumBytes));
+			check(ObjIt.Grab<uint32>() == Magics[3]);
+			FStructSchemaId Id = { ObjIt.Grab<uint32>() };
+			Objects.Add({ { Id, Batch }, ObjIt });
+		}
+		
+		check(It.Grab<uint32>() == Magics[4]);
+		check(!Objects.IsEmpty());
+	}
+
+	~FBatchTextRoundtripper()
+	{
+		UnmountReadSchemas(Batch);
+	}
+
+	void RoundtripText(TConstArrayView<FInstance> Instances) const
+	{
+		PlainProps::UE::RoundtripText(Ids, SavedSchemasView, Batch, Objects, Instances);
+	}
+
+private:
+	FMemoryView					SavedSchemasView;
+	FReadBatchId				Batch;
+	TConstArrayView<FName>		Ids;
+	TArray<FStructView>			Objects;
+	
+	template<typename T>
+	static TConstArrayView<T> GrabNumAndArray(/* in-out */ FByteReader& It)
+	{
+		uint32 Num = It.Grab<uint32>();
+		return MakeArrayView(reinterpret_cast<const T*>(It.GrabBytes(Num * sizeof(T))), Num);
+	}
+};
+
+////////////////////////////////////////////////////////////////////////////////
+
 struct FDiffDebug
 {
 	volatile FInstance* Last;
@@ -3931,7 +4144,7 @@ struct FDiffDebug
 };
 static FDiffDebug GPpDiff, GTpsDiff, GUpsDiff;
 
-static int32 Roundtrip()
+static int32 Roundtrip(ERoundtrip Options)
 {
 	UE_LOGFMT(LogPlainPropsUObject, Display, "Gathering all non-empty UObjects...");
 	static constexpr EObjectFlags SkipFlags = RF_ClassDefaultObject | RF_MirroredGarbage | RF_InheritableComponentTemplate;
@@ -3941,12 +4154,17 @@ static int32 Roundtrip()
 		UObject* Object = *It;
 		if (IncludeClass(Object->GetClass()))
 		{
-			Instances.Add({*It});
+			Instances.Emplace(Object->GetPathName(), Object);
 		}
 	}
+	UE_LOGFMT(LogPlainPropsUObject, Display, "Sorting {Num} UObjects ...", Instances.Num());
+	Algo::Sort(Instances, [](const FInstance& A, const FInstance& B)
+	{
+		return FPlatformString::Strcmp(*A.PathName, *B.PathName) < 0;
+	});
 
 	// Create CDOs if needed and then clones for PP and TPS tests
-	UE_LOGFMT(LogPlainPropsUObject, Display, "Cloning {Num} UObjects 4 times...", Instances.Num());
+	UE_LOGFMT(LogPlainPropsUObject, Display, "Cloning {Num} UObjects up to 4 times...", Instances.Num());
 	for (FInstance& Instance : Instances)
 	{
 		Instance.Init();
@@ -3957,146 +4175,183 @@ static int32 Roundtrip()
 	{
 		uint32 N = 1 + &Instance - Instances.GetData();
 		Instance.Base = Instance.Arch ? MakeEmptyInstance(Instance.Orig, FName("Base", N)) : nullptr;
-		Instance.PP = MakeEmptyInstance(Instance.Orig, FName("PP", N));
-		Instance.TPS = MakeEmptyInstance(Instance.Orig, FName("TPS", N));
-		Instance.UPS = MakeEmptyInstance(Instance.Orig, FName("UPS", N));
+		if (EnumHasAnyFlags(Options, ERoundtrip::PP))
+		{
+			Instance.PP = MakeEmptyInstance(Instance.Orig, FName("PP", N));
+		}
+		if (EnumHasAnyFlags(Options, ERoundtrip::TPS))
+		{
+			Instance.TPS = MakeEmptyInstance(Instance.Orig, FName("TPS", N));
+		}
+		if (EnumHasAnyFlags(Options, ERoundtrip::UPS))
+		{
+			Instance.UPS = MakeEmptyInstance(Instance.Orig, FName("UPS", N));
+		}
 	}
 
 	// Save
-	UE_LOGFMT(LogPlainPropsUObject, Display, "Saving UObjects to PlainProps...");
+	UE_LOGFMT(LogPlainPropsUObject, Display, "Saving UObjects to PlainProps with InMemoryNames...");
 	FMemoryBatch Plain;
 	TCustomPropertyBindings<FMemoryPropertyBatch> Customs(Plain.Properties, GUE.Customs);
 	{
 		FBatchSaver Batch(Customs.Overlay, GUObjectArray.GetObjectArrayNum());
 		SavePlainProps(Batch, Instances);
-		Plain.Data = Batch.Write(/* out */ Plain.RuntimeIds);
+		Plain.Data = Batch.Write(/* out */ &Plain.RuntimeIds);
+
+		if (EnumHasAnyFlags(Options, ERoundtrip::TextMemory))
+		{
+			FBatchTextRoundtripper MemoryBatch(MakeMemoryView(Plain.Data), ESchemaFormat::InMemoryNames);
+			MemoryBatch.RoundtripText(Instances);
+		}
+		if (EnumHasAnyFlags(Options, ERoundtrip::TextStable))
+		{
+			UE_LOGFMT(LogPlainPropsUObject, Display, "Saving UObjects to PlainProps with StableNames...");
+			TArray64<uint8> StableData = Batch.Write();
+
+			FBatchTextRoundtripper StableBatch(MakeMemoryView(StableData), ESchemaFormat::StableNames);
+			StableBatch.RoundtripText(Instances);
+		}
 	}
 
 	// Load
-	UE_LOGFMT(LogPlainPropsUObject, Display, "Loading UObjects from PlainProps...");
-	FMemoryBatchLoader Batch(Customs.Overlay, MakeMemoryView(Plain.Data), Plain.RuntimeIds);
-	LoadPlainProps(Batch, Instances);
-
-	// Diff original vs PlainProps
-	UE_LOGFMT(LogPlainPropsUObject, Display, "Diffing UObjects roundtripped via PlainProps...");
-	FDiffContext DiffCtx = { GUE.Types, GUE.Schemas, Customs.Overlay };
-	TUtf8StringBuilder<256> PpDiffs;
 	uint32 NumPpDiffs = 0;
-	for (FInstance& Instance : Instances)
+	if (EnumHasAnyFlags(Options, ERoundtrip::PP))
 	{
-		if (DiffStructs(Instance.Id, Instance.Orig, Instance.PP, /* in-out */ DiffCtx))
+		UE_LOGFMT(LogPlainPropsUObject, Display, "Loading UObjects from PlainProps...");
+		FMemoryBatchLoader Batch(Customs.Overlay, MakeMemoryView(Plain.Data), Plain.RuntimeIds);
+		LoadPlainProps(Batch, Instances);
+
+		// Diff original vs PlainProps
+		UE_LOGFMT(LogPlainPropsUObject, Display, "Diffing UObjects roundtripped via PlainProps...");
+		FDiffContext DiffCtx = { GUE.Types, GUE.Schemas, Customs.Overlay };
+		TUtf8StringBuilder<256> PpDiffs;
+		for (FInstance& Instance : Instances)
 		{
-			GPpDiff.Last = &Instance;
-			PrintDiff(/* out */ PpDiffs, GUE.Names, DiffCtx.Out);
-			Batch.Reload(Instance.UPS, &Instance - Instances.GetData());
-			DiffCtx.Out.Reset();
-			PpDiffs.Append(" in ");
-			PpDiffs.Append(Instance.Orig->GetFullName());
-			PpDiffs.Append("\n");
-			++NumPpDiffs;
+			if (DiffStructs(Instance.Id, Instance.Orig, Instance.PP, /* in-out */ DiffCtx))
+			{
+				GPpDiff.Last = &Instance;
+				PrintDiff(/* out */ PpDiffs, GUE.Names, DiffCtx.Out);
+				Batch.Reload(Instance.UPS, &Instance - Instances.GetData());
+				DiffCtx.Out.Reset();
+				PpDiffs.Append(" in ");
+				PpDiffs.Append(Instance.Orig->GetFullName());
+				PpDiffs.Append("\n");
+				++NumPpDiffs;
+			}
 		}
+		GPpDiff.Str = PpDiffs.GetData();
+		UE_LOGFMT(LogPlainPropsUObject, Display, "Detected {Diffs} diffs in {Objs} UObjects saved in a {KB}KB value stream using PlainProps\n{DiffText}", NumPpDiffs, Instances.Num(), Plain.Data.NumBytes() / 1024, PpDiffs.ToString());
 	}
-	GPpDiff.Str = PpDiffs.GetData();
-	UE_LOGFMT(LogPlainPropsUObject, Display, "Detected {Diffs} diffs in {Objs} UObjects saved in a {KB}KB value stream using PlainProps", NumPpDiffs, Instances.Num(), Plain.Data.NumBytes() / 1024);
 
-	UE_LOGFMT(LogPlainPropsUObject, Display, "Saving UObjects to TPS archive...");
 	FArchivedProperties Tps;
+	if (EnumHasAnyFlags(Options, ERoundtrip::TPS))
 	{
+		UE_LOGFMT(LogPlainPropsUObject, Display, "Saving UObjects to TPS archive...");
 		FPropertyWriter Archive(/* out */ Tps);
 		SaveArchive<false>(Archive, Instances);
 	}
 
-	UE_LOGFMT(LogPlainPropsUObject, Display, "Loading UObjects from TPS archive...");
+	if (EnumHasAnyFlags(Options, ERoundtrip::TPS))
 	{
+		UE_LOGFMT(LogPlainPropsUObject, Display, "Loading UObjects from TPS archive...");
 		FPropertyReader Archive(Tps);
 		LoadArchive<false>(Archive, Instances);
 	}
 
-	// Diff original vs TPS
-	UE_LOGFMT(LogPlainPropsUObject, Display, "Diffing UObjects roundtripped via TPS...");
 	static const FName SkipClasses[] = {
 		"BodySetup", // Skips some structs due to native FCollisionResponse::operator==
 		"NiagaraScript", "NiagaraNodeFunctionCall", "NiagaraMeshRendererProperties", // FNiagaraTypeDefinition::Serialize resets ClassStructOrEnum
 	};
 
-	TUtf8StringBuilder<256> TpsDiffs;
-	TArray<int32> TpsDiffIdxs;
-	for (FInstance& Instance : Instances)
+	if (EnumHasAnyFlags(Options, ERoundtrip::TPS))
 	{
-		if (Algo::Find(SkipClasses, Instance.Orig->GetClass()->GetFName()))
-		{
-			continue;
-		}
-
-		if (DiffStructs(Instance.Id, Instance.Orig, Instance.TPS, /* in-out */ DiffCtx))
+		// Diff original vs TPS
+		UE_LOGFMT(LogPlainPropsUObject, Display, "Diffing UObjects roundtripped via TPS...");
+		FDiffContext DiffCtx = { GUE.Types, GUE.Schemas, Customs.Overlay };
+		TUtf8StringBuilder<256> TpsDiffs;
+		TArray<int32> TpsDiffIdxs;
+		for (FInstance& Instance : Instances)
 		{
-			GTpsDiff.Last = &Instance;
-			PrintDiff(/* out */ TpsDiffs, GUE.Names, DiffCtx.Out);
-			
-			DiffCtx.Out.Reset();
-			TpsDiffs.Append(" in ");
-			TpsDiffs.Append(Instance.Orig->GetFullName());
-			TpsDiffs.Append("\n");	
-			TpsDiffIdxs.Add(&Instance - &Instances[0]);
-			
-			//FArchivedProperties Tmp;
-			//FPropertyWriter(/* out */ Tmp).WriteProperties(Instance.Orig, Instance.Base);
-			//FPropertyReader(Tmp).ReadProperties(Instance.UPS, Instance.Base);
+			if (Algo::Find(SkipClasses, Instance.Orig->GetClass()->GetFName()))
+			{
+				continue;
+			}
 
-			//FDiffContext TmpCtx = DiffCtx;
-			//bool bOU = DiffStructs(Instance.Id, Instance.Orig, Instance.UPS, /* in-out */ TmpCtx);
-			//bool bOT = DiffStructs(Instance.Id, Instance.Orig, Instance.TPS, /* in-out */ TmpCtx);
-			//bool bOP = DiffStructs(Instance.Id, Instance.Orig, Instance.PP, /* in-out */ TmpCtx);
-			//TUtf8StringBuilder<256> TmpDiff;
-			//PrintDiff(/* out */ TmpDiff, TmpCtx.Out, GUE.Names);
+			if (DiffStructs(Instance.Id, Instance.Orig, Instance.TPS, /* in-out */ DiffCtx))
+			{
+				GTpsDiff.Last = &Instance;
+				PrintDiff(/* out */ TpsDiffs, GUE.Names, DiffCtx.Out);
+				
+				DiffCtx.Out.Reset();
+				TpsDiffs.Append(" in ");
+				TpsDiffs.Append(Instance.Orig->GetFullName());
+				TpsDiffs.Append("\n");	
+				TpsDiffIdxs.Add(&Instance - &Instances[0]);
+				
+				//FArchivedProperties Tmp;
+				//FPropertyWriter(/* out */ Tmp).WriteProperties(Instance.Orig, Instance.Base);
+				//FPropertyReader(Tmp).ReadProperties(Instance.UPS, Instance.Base);
+
+				//FDiffContext TmpCtx = DiffCtx;
+				//bool bOU = DiffStructs(Instance.Id, Instance.Orig, Instance.UPS, /* in-out */ TmpCtx);
+				//bool bOT = DiffStructs(Instance.Id, Instance.Orig, Instance.TPS, /* in-out */ TmpCtx);
+				//bool bOP = DiffStructs(Instance.Id, Instance.Orig, Instance.PP, /* in-out */ TmpCtx);
+				//TUtf8StringBuilder<256> TmpDiff;
+				//PrintDiff(/* out */ TmpDiff, TmpCtx.Out, GUE.Names);
+			}
 		}
+		GTpsDiff.Str = TpsDiffs.GetData();
+		UE_LOGFMT(LogPlainPropsUObject, Display, "Detected {Diffs} diffs in {Objs} UObjects saved in a {KB}KB value stream using TPS", TpsDiffIdxs.Num(), Instances.Num(), Tps.Data.NumBytes() / 1024);
 	}
-	GTpsDiff.Str = TpsDiffs.GetData();
-	UE_LOGFMT(LogPlainPropsUObject, Display, "Detected {Diffs} diffs in {Objs} UObjects saved in a {KB}KB value stream using TPS", TpsDiffIdxs.Num(), Instances.Num(), Tps.Data.NumBytes() / 1024);
 
-	UE_LOGFMT(LogPlainPropsUObject, Display, "Saving UObjects to UPS archive...");
-	FArchivedProperties Ups;
-	{
-		FPropertyWriter Archive(/* out */ Ups);
-		SaveArchive<true>(Archive, Instances);
-	}
 
-	UE_LOGFMT(LogPlainPropsUObject, Display, "Loading UObjects from UPS archive...");
+	if (EnumHasAnyFlags(Options, ERoundtrip::UPS))
 	{
-		FPropertyReader Archive(Ups);
-		LoadArchive<true>(Archive, Instances);
-	}
+		FArchivedProperties Ups;
+		UE_LOGFMT(LogPlainPropsUObject, Display, "Saving UObjects to UPS archive...");
+		{
+			FPropertyWriter Archive(/* out */ Ups);
+			SaveArchive<true>(Archive, Instances);
+		}
 
-	// Diff original vs UPS
-	UE_LOGFMT(LogPlainPropsUObject, Display, "Diffing UObjects roundtripped via UPS...");
-	TUtf8StringBuilder<256> UpsDiffs;
-	TArray<int32> UpsDiffIdxs;
-	for (FInstance& Instance : Instances)
-	{
-		if (Algo::Find(SkipClasses, Instance.Orig->GetClass()->GetFName()))
+		UE_LOGFMT(LogPlainPropsUObject, Display, "Loading UObjects from UPS archive...");
 		{
-			continue;
+			FPropertyReader Archive(Ups);
+			LoadArchive<true>(Archive, Instances);
 		}
 
-		if (DiffStructs(Instance.Id, Instance.Orig, Instance.UPS, /* in-out */ DiffCtx))
+		// Diff original vs UPS
+		UE_LOGFMT(LogPlainPropsUObject, Display, "Diffing UObjects roundtripped via UPS...");
+		FDiffContext DiffCtx = { GUE.Types, GUE.Schemas, Customs.Overlay };
+		TUtf8StringBuilder<256> UpsDiffs;
+		TArray<int32> UpsDiffIdxs;
+		for (FInstance& Instance : Instances)
 		{
-			GUpsDiff.Last = &Instance;
-			PrintDiff(/* out */ UpsDiffs, GUE.Names, DiffCtx.Out);
-			
-			DiffCtx.Out.Reset();
-			UpsDiffs.Append(" in ");
-			UpsDiffs.Append(Instance.Orig->GetFullName());
-			UpsDiffs.Append("\n");	
-			UpsDiffIdxs.Add(&Instance - &Instances[0]);
+			if (Algo::Find(SkipClasses, Instance.Orig->GetClass()->GetFName()))
+			{
+				continue;
+			}
+
+			if (DiffStructs(Instance.Id, Instance.Orig, Instance.UPS, /* in-out */ DiffCtx))
+			{
+				GUpsDiff.Last = &Instance;
+				PrintDiff(/* out */ UpsDiffs, GUE.Names, DiffCtx.Out);
+				
+				DiffCtx.Out.Reset();
+				UpsDiffs.Append(" in ");
+				UpsDiffs.Append(Instance.Orig->GetFullName());
+				UpsDiffs.Append("\n");	
+				UpsDiffIdxs.Add(&Instance - &Instances[0]);
+			}
 		}
+		GUpsDiff.Str = UpsDiffs.GetData();
+		UE_LOGFMT(LogPlainPropsUObject, Display, "Detected {Diffs} diffs in {Objs} UObjects saved in a {KB}KB value stream using UPS", UpsDiffIdxs.Num(), Instances.Num(), Ups.Data.NumBytes() / 1024);
 	}
-	GUpsDiff.Str = UpsDiffs.GetData();
-	UE_LOGFMT(LogPlainPropsUObject, Display, "Detected {Diffs} diffs in {Objs} UObjects saved in a {KB}KB value stream using UPS", UpsDiffIdxs.Num(), Instances.Num(), Ups.Data.NumBytes() / 1024);
 
 	return NumPpDiffs;
 }
 
-static int32 TestBindings()
+static int32 TestBindings(ERoundtrip Options)
 {
 	TScopedStructBinding<FTransform, FDefaultRuntime> Transform;
 	TScopedStructBinding<FGuid, FDefaultRuntime> Guid;
@@ -4115,7 +4370,7 @@ static int32 TestBindings()
 
 	InitBatchedProperties();
 	BindInitialTypes();
-	return Roundtrip();
+	return Roundtrip(Options);
 }
 
 } // namespace PlainProps::UE
@@ -4169,7 +4424,22 @@ int32 UTestPlainPropsCommandlet::Main(const FString& Params)
 	FlushAsyncLoading();
 
 	UE_LOGFMT(LogPlainPropsUObject, Display, "Starting test...");
-	return PlainProps::UE::TestBindings();
+
+	using namespace PlainProps::UE;
+	ERoundtrip Options = ERoundtrip::PP | ERoundtrip::UPS | ERoundtrip::TPS | ERoundtrip::TextMemory;
+	if (Params.Find("-pp") != INDEX_NONE)
+	{
+		Options = ERoundtrip::PP | ERoundtrip::TextMemory;
+	}
+	else if (Params.Find("-text") != INDEX_NONE)
+	{
+		Options = ERoundtrip::TextMemory | ERoundtrip::TextStable;
+	}
+	else if (Params.Find("-notext") != INDEX_NONE)
+	{
+		EnumRemoveFlags(Options, ERoundtrip::TextMemory | ERoundtrip::TextStable);
+	}
+	return TestBindings(Options);
 }
 
 // Stopgap debugger callable functions til natvis is added 
@@ -4181,4 +4451,4 @@ extern FString _print(const PlainProps::FTypenameId& Id)			{ return PlainProps::
 extern FString _print(const PlainProps::FConcreteTypenameId& Id)	{ return PlainProps::UE::GUE.Debug.Print(Id); }
 extern FString _print(const PlainProps::FType& Id)					{ return PlainProps::UE::GUE.Debug.Print(Id); }
 extern FString _print(const PlainProps::FEnumId& Id)				{ return PlainProps::UE::GUE.Debug.Print(Id); }
-extern FString _print(const PlainProps::FStructId& Id)				{ return PlainProps::UE::GUE.Debug.Print(Id); }
\ No newline at end of file
+extern FString _print(const PlainProps::FStructId& Id)				{ return PlainProps::UE::GUE.Debug.Print(Id); }
