PlainProps: Compile fix in PlainPropsUObjectBindings.cpp for cl 39626420

#rnx

diff --git a/Engine/Plugins/Experimental/PlainPropsUObject/Source/Private/PlainPropsUObjectBindings.cpp b/Engine/Plugins/Experimental/PlainPropsUObject/Source/Private/PlainPropsUObjectBindings.cpp
--- a/Engine/Plugins/Experimental/PlainPropsUObject/Source/Private/PlainPropsUObjectBindings.cpp
+++ b/Engine/Plugins/Experimental/PlainPropsUObject/Source/Private/PlainPropsUObjectBindings.cpp
@@ -4011,7 +4011,7 @@ static void RoundtripText(
 	// Print yaml
 	UE_LOGFMT(LogPlainPropsUObject, Display, "Printing to PlainProps text using {Format}...", ToString(Format));
 	TUtf8StringBuilder<256> Yaml;
-	Yaml.Reserve((2 << 30) - 1);
+	Yaml.Reserve(INT_MAX);
 	if (Format == ESchemaFormat::InMemoryNames)
 	{
 		FMemoryBatchIds Ids(Batch, GUE.Names);
