PlainProps disable accidental test code

#rnx

diff --git a/Engine/Plugins/Experimental/PlainPropsUObject/Source/Private/PlainPropsUObjectBindings.cpp b/Engine/Plugins/Experimental/PlainPropsUObject/Source/Private/PlainPropsUObjectBindings.cpp
--- a/Engine/Plugins/Experimental/PlainPropsUObject/Source/Private/PlainPropsUObjectBindings.cpp
+++ b/Engine/Plugins/Experimental/PlainPropsUObject/Source/Private/PlainPropsUObjectBindings.cpp
@@ -4231,7 +4231,7 @@ static int32 Roundtrip(ERoundtrip Options)
 			{
 				GPpDiff.Last = &Instance;
 				PrintDiff(/* out */ PpDiffs, GUE.Names, DiffCtx.Out);
-				Batch.Reload(Instance.UPS, &Instance - Instances.GetData());
+				//Batch.Reload(Instance.UPS, &Instance - Instances.GetData());
 				DiffCtx.Out.Reset();
 				PpDiffs.Append(" in ");
 				PpDiffs.Append(Instance.Orig->GetFullName());
