PlainProps: Add "CitySampleUnrealEdEngine" to commandlet exclude list due to a MTAccessDetector crash while cloning.

#rnx

diff --git a/Engine/Plugins/Experimental/PlainPropsUObject/Source/Private/PlainPropsUObjectBindings.cpp b/Engine/Plugins/Experimental/PlainPropsUObject/Source/Private/PlainPropsUObjectBindings.cpp
--- a/Engine/Plugins/Experimental/PlainPropsUObject/Source/Private/PlainPropsUObjectBindings.cpp
+++ b/Engine/Plugins/Experimental/PlainPropsUObject/Source/Private/PlainPropsUObjectBindings.cpp
@@ -3894,10 +3894,11 @@ static UObject* MakeEmptyInstance(UObject* Obj, FName Name)
 static bool IncludeClass(UClass* Class)
 {
 	static const FName Exclusions[] = {
+		"CitySampleUnrealEdEngine", // Cloning MTAccessDetector crash
 		"GameFeaturePluginStateMachine", // Cloning ensure
 		"WorldSettings", // QAGame
 		// CitySample - Enum value 8 is undeclared in /Script/Engine.ERichCurveTangentMode, illegal value detected in /Script/Engine.RichCurveKey::TangentMode
-		"AnimSequence", "AnimationSequencerDataModel", "MovieSceneControlRigParameterSection" 
+		"AnimSequence", "AnimationSequencerDataModel", "MovieSceneControlRigParameterSection"
 	};
 	static const FName SuperExclusions[] = { "LevelScriptActor" };
 
