PlainProps split FStructId into FDeclId and FBindId

#rnx

diff --git a/Engine/Plugins/Experimental/PlainProps/Source/Private/PlainPropsBind.cpp b/Engine/Plugins/Experimental/PlainProps/Source/Private/PlainPropsBind.cpp
--- a/Engine/Plugins/Experimental/PlainProps/Source/Private/PlainPropsBind.cpp
+++ b/Engine/Plugins/Experimental/PlainProps/Source/Private/PlainPropsBind.cpp
@@ -71,7 +71,7 @@ FLeafMemberBinding FMemberVisitor::GrabLeaf()
 FStructMemberBinding FMemberVisitor::GrabStruct()
 {
 	checkf(!PeekType().AsStruct().IsDynamic, TEXT("Bound structs can't be dynamic"));
-	return { PeekType().AsStruct(), GrabInnerSchema().AsStruct(), GrabMemberOffset() };
+	return { PeekType().AsStruct(), GrabInnerSchema().AsStructBindId(), GrabMemberOffset() };
 }
 
 static bool HasSchema(FMemberBindType Type)
@@ -115,13 +115,13 @@ void FMemberVisitor::SkipMember()
 	++MemberIdx;
 }
 
-FStructId FMemberVisitor::GrabSuper()
+FBindId FMemberVisitor::GrabSuper()
 {
 	check(MemberIdx == 0);
 	checkSlow(Schema.Members[0].AsByte() == SuperStructType.AsByte());
 	MemberIdx = 1;
 	InnerSchemaIdx = 1;
-	return Schema.GetInnerSchemas()[0].AsStruct();
+	return Schema.GetInnerSchemas()[0].AsStructBindId();
 }
 
 FInnerId FMemberVisitor::GrabInnerSchema()
@@ -197,12 +197,12 @@ struct FCustomBindingHandle
 	FInnersHandle*					GetLoweredInners() { return reinterpret_cast<FInnersHandle*>(Bindings + 2 * Max); }
 };
 
-void FCustomBindingMap::Bind(FStructId BindId, ICustomBinding& Binding, const FStructDeclaration& Declaration, FInnersHandle LoweredInners)
+void FCustomBindingMap::Bind(FBindId Id, ICustomBinding& Binding, const FStructDeclaration& Declaration, FInnersHandle LoweredInners)
 {
 	// Add key
 	bool bExists = false;
-	int32 SetIdx = Keys.Add(BindId, &bExists).AsInteger();
-	checkf(!bExists, TEXT("'%s' already bound"), *Debug.Print(BindId));
+	int32 SetIdx = Keys.Add(Id, &bExists).AsInteger();
+	checkf(!bExists, TEXT("'%s' already bound"), *Debug.Print(Id));
 
 	// Grow values if needed
 	uint32 MaxKeys = static_cast<uint32>(Keys.GetMaxIndex());
@@ -230,14 +230,14 @@ void FCustomBindingMap::Bind(FStructId BindId, ICustomBinding& Binding, const FS
 	*Value.GetDeclarations() = &Declaration;
 	*Value.GetLoweredInners() = LoweredInners;
 
-	UpdateCulling(/* in-out */ Window, BindId.Idx);
+	UpdateCulling(/* in-out */ Window, Id.Idx);
 }
 
-inline FCustomBindingHandle FCustomBindingMap::Find(FStructId BindId) const
+inline FCustomBindingHandle FCustomBindingMap::Find(FBindId Id) const
 {
-	if (Cull(Window, BindId.Idx))
+	if (Cull(Window, Id.Idx))
 	{
-		if (FSetElementId Idx = Keys.FindId(BindId); Idx.IsValidId())
+		if (FSetElementId Idx = Keys.FindId(Id); Idx.IsValidId())
 		{
 			return {Values + Idx.AsInteger(), MaxValues};
 		}
@@ -245,10 +245,10 @@ inline FCustomBindingHandle FCustomBindingMap::Find(FStructId BindId) const
 	return {};
 }
 
-void FCustomBindingMap::Drop(FStructId BindId)
+void FCustomBindingMap::Drop(FBindId Id)
 {
-	int32 NumRemoved = Keys.Remove(BindId);
-	checkf(NumRemoved == 1, TEXT("'%s' unbound"), *Debug.Print(BindId));
+	int32 NumRemoved = Keys.Remove(Id);
+	checkf(NumRemoved == 1, TEXT("'%s' unbound"), *Debug.Print(Id));
 }
 
 ////////////////////////////////////////////////////////////////////////////////////////////////
@@ -276,26 +276,26 @@ static TConstArrayView<FInnerStruct> FetchInners(const TArray<FInnerStruct>& All
 	return {};
 }
 
-static bool operator==(FInnerStruct A, FInnerStruct B) { return A.Name == B.Name && A.BindId == B.BindId; }
+static bool operator==(FInnerStruct A, FInnerStruct B) { return A.Name == B.Name && A.Id == B.Id; }
 
 ////////////////////////////////////////////////////////////////////////////////////////////////
 
-void FCustomBindings::BindStruct(FStructId BindId, ICustomBinding& Binding, const FStructDeclaration& Declaration, TConstArrayView<FInnerStruct> LoweredInners)
+void FCustomBindings::BindStruct(FBindId Id, ICustomBinding& Binding, const FStructDeclaration& Declaration, TConstArrayView<FInnerStruct> LoweredInners)
 {
 	FInnersHandle StoredInners = StoreInners(/* out */ BottomInners, LoweredInners);
 	check(Algo::Compare(LoweredInners, FetchInners(BottomInners, StoredInners)));
-	Map.Bind(BindId, Binding, Declaration, StoredInners);
+	Map.Bind(Id, Binding, Declaration, StoredInners);
 }
 
-const ICustomBinding* FCustomBindings::FindStruct(FStructId BindId) const
+const ICustomBinding* FCustomBindings::FindStruct(FBindId Id) const
 {
-	FCustomBindingHandle Handle = Find(BindId);
+	FCustomBindingHandle Handle = Find(Id);
 	return Handle ? *Handle.Bindings : nullptr;
 }
 
-const ICustomBinding* FCustomBindings::FindStruct(FStructId BindId, TConstArrayView<FInnerStruct>& OutLoweredInners) const
+const ICustomBinding* FCustomBindings::FindStruct(FBindId Id, TConstArrayView<FInnerStruct>& OutLoweredInners) const
 {
-	if (FCustomBindingHandle Handle = Find(BindId))
+	if (FCustomBindingHandle Handle = Find(Id))
 	{
 		OutLoweredInners = FetchInners(BottomInners, *Handle.GetLoweredInners());
 		return *Handle.Bindings;
@@ -303,33 +303,33 @@ const ICustomBinding* FCustomBindings::FindStruct(FStructId BindId, TConstArrayV
 	return nullptr;
 }
 
-ICustomBinding* FCustomBindings::FindStructToSave(FStructId BindId, const FStructDeclaration*& OutDeclaration) const
+ICustomBinding* FCustomBindings::FindStructToSave(FBindId Id, const FStructDeclaration*& OutDeclaration) const
 {
-	FCustomBindingHandle Handle = Find(BindId);
+	FCustomBindingHandle Handle = Find(Id);
 	OutDeclaration = Handle ? *Handle.GetDeclarations() : nullptr;
 	return Handle ? *Handle.Bindings : nullptr;
 }
 
-const FStructDeclaration* FCustomBindings::FindDeclaration(FStructId BindId) const
+const FStructDeclaration* FCustomBindings::FindDeclaration(FBindId Id) const
 { 
-	FCustomBindingHandle Handle = Find(BindId);
+	FCustomBindingHandle Handle = Find(Id);
 	return Handle ? *Handle.GetDeclarations() : nullptr;
 }
 
-void FCustomBindings::DropStruct(FStructId BindId)
+void FCustomBindings::DropStruct(FBindId Id)
 {
-	Map.Drop(BindId);
+	Map.Drop(Id);
 }
 
-FCustomBindingHandle FCustomBindingsBottom::Find(FStructId BindId) const
+FCustomBindingHandle FCustomBindingsBottom::Find(FBindId Id) const
 {
-	return Map.Find(BindId);
+	return Map.Find(Id);
 }
 
-FCustomBindingHandle FCustomBindingsOverlay::Find(FStructId BindId) const
+FCustomBindingHandle FCustomBindingsOverlay::Find(FBindId Id) const
 {
-	FCustomBindingHandle Handle = Map.Find(BindId);
-	return Handle ? Handle : Underlay.Find(BindId);
+	FCustomBindingHandle Handle = Map.Find(Id);
+	return Handle ? Handle : Underlay.Find(Id);
 }
 
 ////////////////////////////////////////////////////////////////////////////////////////////////
@@ -350,7 +350,7 @@ inline FType IndexInnermost(FIdIndexerBase& Ids, FMemberBindType Type, FOptional
 	return Ids.MakeLeafParameter(Leaf);
 }
 
-FDualType FMemberBinding::IndexParameterName(FIdIndexerBase& Ids) const
+FBothType FMemberBinding::IndexParameterName(FIdIndexerBase& Ids) const
 {
 	FType BindType = IndexInnermost(Ids, InnermostType, InnermostSchema);
 	FType DeclType = BindType;
@@ -408,7 +408,7 @@ struct FMemberBinder : FMemberBinderBase
 	FInnerId* InnerSchemaIt;
 };
 
-void FSchemaBindings::BindStruct(FStructId BindId, FStructId DeclId, TConstArrayView<FMemberBinding> Members)
+void FSchemaBindings::BindStruct(FBindId BindId, FDeclId DeclId, TConstArrayView<FMemberBinding> Members)
 {
 	// Make header, allocate and copy header
 	FSchemaBinding Header = { DeclId, IntCastChecked<uint16>(Members.Num()), CountInnerSchemas(Members), CountRanges(Members) };
@@ -443,34 +443,34 @@ void FSchemaBindings::BindStruct(FStructId BindId, FStructId DeclId, TConstArray
 	Bindings[BindId.Idx].Reset(Schema);
 }
 
-const FSchemaBinding* FSchemaBindings::FindStruct(FStructId BindId) const
+const FSchemaBinding* FSchemaBindings::FindStruct(FBindId Id) const
 {
-	return BindId.Idx < (uint32)Bindings.Num() ? Bindings[BindId.Idx].Get() : nullptr;
+	return Id.Idx < (uint32)Bindings.Num() ? Bindings[Id.Idx].Get() : nullptr;
 }
 
-const FSchemaBinding& FSchemaBindings::GetStruct(FStructId BindId) const
+const FSchemaBinding& FSchemaBindings::GetStruct(FBindId Id) const
 {
-	checkf(BindId.Idx < (uint32)Bindings.Num() && Bindings[BindId.Idx], TEXT("'%s' is unbound"), *Debug.Print(BindId));
-	return *Bindings[BindId.Idx].Get();
+	checkf(Id.Idx < (uint32)Bindings.Num() && Bindings[Id.Idx], TEXT("'%s' is unbound"), *Debug.Print(Id));
+	return *Bindings[Id.Idx].Get();
 }
 
-void FSchemaBindings::DropStruct(FStructId BindId)
+void FSchemaBindings::DropStruct(FBindId Id)
 {
-	checkf(BindId.Idx < (uint32)Bindings.Num() && Bindings[BindId.Idx], TEXT("'%s' is unbound"), *Debug.Print(BindId));
-	Bindings[BindId.Idx].Reset();
+	checkf(Id.Idx < (uint32)Bindings.Num() && Bindings[Id.Idx], TEXT("'%s' is unbound"), *Debug.Print(Id));
+	Bindings[Id.Idx].Reset();
 }
 
-FStructId FSchemaBindings::GetDeclId(FStructId BindId) const
+FDeclId FSchemaBindings::Lower(FBindId Id) const
 {
-	return GetStruct(BindId).DeclId;
+	return GetStruct(Id).DeclId;
 }
 
 //////////////////////////////////////////////////////////////////////////
 
-FStructId FStructBindIds::GetDeclId(FStructId BindId) const
+FDeclId FStructBindIds::Lower(FBindId Id) const
 {
-	const FStructDeclaration* CustomDecl = Customs.FindDeclaration(BindId);
-	return CustomDecl ? CustomDecl->Id : Schemas.GetDeclId(BindId);
+	const FStructDeclaration* CustomDecl = Customs.FindDeclaration(Id);
+	return CustomDecl ? CustomDecl->Id : Schemas.Lower(Id);
 }
 
 //////////////////////////////////////////////////////////////////////////
diff --git a/Engine/Plugins/Experimental/PlainProps/Source/Private/PlainPropsBuildSchema.cpp b/Engine/Plugins/Experimental/PlainProps/Source/Private/PlainPropsBuildSchema.cpp
--- a/Engine/Plugins/Experimental/PlainProps/Source/Private/PlainPropsBuildSchema.cpp
+++ b/Engine/Plugins/Experimental/PlainProps/Source/Private/PlainPropsBuildSchema.cpp
@@ -74,11 +74,11 @@ struct FEnumSchemaBuilder
 
 //////////////////////////////////////////////////////////////////////////
 
-FSchemasBuilder::FSchemasBuilder(const FDeclarations& Types, const FIds& InIds, const IStructBindIds& InBindIds, FScratchAllocator& InScratch, ESchemaFormat InFormat)
+FSchemasBuilder::FSchemasBuilder(const FDeclarations& Types, const FIds& InIds, const IBindIds& InBindIds, FScratchAllocator& InScratch, ESchemaFormat InFormat)
 : FSchemasBuilder(Types.GetStructs(), Types.GetEnums(), InIds, InBindIds, InScratch, InFormat)
 {}
 
-FSchemasBuilder::FSchemasBuilder(FStructDeclarations InStructs, FEnumDeclarations InEnums, const FIds& InIds, const IStructBindIds& InBindIds, FScratchAllocator& InScratch, ESchemaFormat InFormat)
+FSchemasBuilder::FSchemasBuilder(FStructDeclarations InStructs, FEnumDeclarations InEnums, const FIds& InIds, const IBindIds& InBindIds, FScratchAllocator& InScratch, ESchemaFormat InFormat)
 : DeclaredStructs(InStructs)
 , DeclaredEnums(InEnums)
 , Ids(InIds)
@@ -112,21 +112,21 @@ FORCEINLINE FEnumSchemaBuilder&	FSchemasBuilder::NoteEnum(FEnumId Id)
 	return GetOrEmplace(EnumIndices, Id, Enums, *DeclaredEnums[Id.Idx], Debug, Id);
 }
 
-FORCEINLINE FStructSchemaBuilder& FSchemasBuilder::NoteStruct(FStructId BindId)
+FORCEINLINE FStructSchemaBuilder& FSchemasBuilder::NoteStruct(FStructId Id)
 {
 	checkf(!bBuilt, TEXT("Noted new members after building"));
 
-	FStructId DeclId	= BindId.Idx < uint32(DeclaredStructs.Num()) && DeclaredStructs[BindId.Idx] 
-							? BindId : BindIds.GetDeclId(BindId);
+	FDeclId DeclId	= Id.Idx < uint32(DeclaredStructs.Num()) && DeclaredStructs[Id.Idx] 
+					? FDeclId(Id) : BindIds.Lower(FBindId(Id));
 
 	checkf(DeclId.Idx < uint32(DeclaredStructs.Num()) && DeclaredStructs[DeclId.Idx], TEXT("Undeclared struct '%s' noted"), *Debug.Print(DeclId));
-	FStructId NoteId = (Format == ESchemaFormat::StableNames) ? DeclId : BindId;
+	FStructId NoteId = (Format == ESchemaFormat::StableNames) ? DeclId : Id;
 	return GetOrEmplace(StructIndices, NoteId, Structs, NoteId, *DeclaredStructs[DeclId.Idx], *this);
 }
 
-void FSchemasBuilder::NoteStructAndMembers(FStructId BindId, const FBuiltStruct& Struct)
+void FSchemasBuilder::NoteStructAndMembers(FStructId Id, const FBuiltStruct& Struct)
 {
-	NoteStruct(BindId).NoteMembersRecursively(Struct);
+	NoteStruct(Id).NoteMembersRecursively(Struct);
 }
 
 FBuiltSchemas FSchemasBuilder::Build()
@@ -156,7 +156,7 @@ void FSchemasBuilder::NoteInheritanceChains()
 {
 	for (int Idx = 0, Num = Structs.Num(); Idx < Num; ++Idx)
 	{
-		for (FOptionalStructId Super = Structs[Idx].Declaration.Super; Super; Super = DeclaredStructs[Super.Get().Idx]->Super)
+		for (FOptionalDeclId Super = Structs[Idx].Declaration.Super; Super; Super = DeclaredStructs[Super.Get().Idx]->Super)
 		{
 			FStructId Id = Super.Get();
 			GetOrEmplace(StructIndices, Id, Structs, Id, *DeclaredStructs[Id.Idx], *this);
@@ -211,7 +211,7 @@ static void SetIsDynamic(FMemberType& InOut)
 
 static void* NoteStructOrEnum(FSchemasBuilder& AllSchemas, bool bStruct, FInnerId Id)
 {
-	return bStruct ? static_cast<void*>(&AllSchemas.NoteStruct(Id.AsStruct())) : &AllSchemas.NoteEnum(Id.AsEnum());
+	return bStruct ? static_cast<void*>(&AllSchemas.NoteStruct(Id.AsStructBindId())) : &AllSchemas.NoteEnum(Id.AsEnum());
 }
 
 void FStructSchemaBuilder::NoteMembersRecursively(const FBuiltStruct& Struct)
@@ -278,7 +278,7 @@ void FStructSchemaBuilder::NoteMembersRecursively(const FBuiltStruct& Struct)
 				break;
 
 			case EMemberKind::Struct:
-				AllSchemas.NoteStruct(InnerSchema.AsStruct()).NoteMembersRecursively(*Member.Value.Struct);
+				AllSchemas.NoteStruct(InnerSchema.AsStructBindId()).NoteMembersRecursively(*Member.Value.Struct);
 				break;
 		
 			case EMemberKind::Range:
@@ -352,7 +352,7 @@ void FStructSchemaBuilder::NoteRangeRecursively(ERangeSizeType NumType, TConstAr
 FBuiltStructSchema FStructSchemaBuilder::Build() const
 {
 	FType Type = Id == Declaration.Id ? Declaration.Type : AllSchemas.GetIds().Resolve(Id);
-	FBuiltStructSchema Out = { Type, Id, Declaration.Super, /* dense */ true };
+	FBuiltStructSchema Out = { Type, Id, ToOptionalStruct(Declaration.Super), /* dense */ true };
 	
 	if (int32 Num = NotedMembers.CountSetBits())
 	{
diff --git a/Engine/Plugins/Experimental/PlainProps/Source/Private/PlainPropsDeclare.cpp b/Engine/Plugins/Experimental/PlainProps/Source/Private/PlainPropsDeclare.cpp
--- a/Engine/Plugins/Experimental/PlainProps/Source/Private/PlainPropsDeclare.cpp
+++ b/Engine/Plugins/Experimental/PlainProps/Source/Private/PlainPropsDeclare.cpp
@@ -97,17 +97,17 @@ void CopyItems(T* It, TConstArrayView<T> Items)
 	}
 }
 
-static const FStructDeclaration& Declare(TArray<TUniquePtr<FStructDeclaration>>& Out, FStructId DeclId, FType Type, uint16 Version, TConstArrayView<FMemberId> MemberOrder, EMemberPresence Occupancy, FOptionalStructId Super, uint16 NumMembers)
+static const FStructDeclaration& Declare(TArray<TUniquePtr<FStructDeclaration>>& Out, FDeclId Id, FType Type, uint16 Version, TConstArrayView<FMemberId> MemberOrder, EMemberPresence Occupancy, FOptionalDeclId Super, uint16 NumMembers)
 {
-	if (static_cast<int32>(DeclId.Idx) >= Out.Num())
+	if (static_cast<int32>(Id.Idx) >= Out.Num())
 	{
-		Out.SetNum(DeclId.Idx + 1);
+		Out.SetNum(Id.Idx + 1);
 	}
 
-	TUniquePtr<FStructDeclaration>& Ptr = Out[DeclId.Idx];
+	TUniquePtr<FStructDeclaration>& Ptr = Out[Id.Idx];
 	if (Ptr)
 	{
-		check(DeclId == Ptr->Id);
+		check(Id == Ptr->Id);
 		check(Type == Ptr->Type);
 		check(Super == Ptr->Super);
 		check(Version == Ptr->Version);
@@ -119,7 +119,7 @@ static const FStructDeclaration& Declare(TArray<TUniquePtr<FStructDeclaration>>&
 	}
 	else
 	{
-		FStructDeclaration Header{1, DeclId, Type, Super, Version, NumMembers, Occupancy};
+		FStructDeclaration Header{1, Id, Type, Super, Version, NumMembers, Occupancy};
 		void* Data = FMemory::Malloc(sizeof(FStructDeclaration) + MemberOrder.Num() * MemberOrder.GetTypeSize());
 		Ptr.Reset(new (Data) FStructDeclaration(Header));
 		CopyItems(Ptr->MemberOrder, MemberOrder);
@@ -128,15 +128,15 @@ static const FStructDeclaration& Declare(TArray<TUniquePtr<FStructDeclaration>>&
 	return *Ptr;
 }
 
-const FStructDeclaration& FDeclarations::DeclareStruct(FStructId DeclId, FType Type, uint16 Version, TConstArrayView<FMemberId> MemberOrder, EMemberPresence Occupancy, FOptionalStructId Super)
+const FStructDeclaration& FDeclarations::DeclareStruct(FDeclId Id, FType Type, uint16 Version, TConstArrayView<FMemberId> MemberOrder, EMemberPresence Occupancy, FOptionalDeclId Super)
 {
-	checkf(!(Super && Occupancy == EMemberPresence::RequireAll), TEXT("'%s' is a dense substruct, this isn't supported, see BuildSuperStruct()"), *Debug.Print(DeclId));
-	return Declare(/* out */ DeclaredStructs, DeclId, Type, Version, MemberOrder, Occupancy, Super, IntCastChecked<uint16>(MemberOrder.Num()));
+	checkf(!(Super && Occupancy == EMemberPresence::RequireAll), TEXT("'%s' is a dense substruct, this isn't supported, see BuildSuperStruct()"), *Debug.Print(Id));
+	return Declare(/* out */ DeclaredStructs, Id, Type, Version, MemberOrder, Occupancy, Super, IntCastChecked<uint16>(MemberOrder.Num()));
 }
 
-const FStructDeclaration& FDeclarations::DeclareNumeralStruct(FStructId DeclId, FType Type, TConstArrayView<FMemberId> Numerals, EMemberPresence Occupancy)
+const FStructDeclaration& FDeclarations::DeclareNumeralStruct(FDeclId Id, FType Type, TConstArrayView<FMemberId> Numerals, EMemberPresence Occupancy)
 {
-	return Declare(/* out */ DeclaredStructs, DeclId, Type, 0, Numerals, Occupancy, NoId, IntCastChecked<uint16>(Numerals.Num()));
+	return Declare(/* out */ DeclaredStructs, Id, Type, 0, Numerals, Occupancy, NoId, IntCastChecked<uint16>(Numerals.Num()));
 }
 
 const FEnumDeclaration& FDeclarations::DeclareEnum(FEnumId Id, FType Type, EEnumMode Mode, TConstArrayView<FEnumerator> Enumerators, EEnumAliases Policy)
@@ -165,11 +165,11 @@ const FEnumDeclaration& FDeclarations::DeclareEnum(FEnumId Id, FType Type, EEnum
 	return *Ptr;
 }
 
-void FDeclarations::DropStructRef(FStructId DeclId)
+void FDeclarations::DropStructRef(FDeclId Id)
 {
-	Check(DeclId);
+	Check(Id);
 	
-	TUniquePtr<FStructDeclaration>& Ptr = DeclaredStructs[DeclId.Idx];
+	TUniquePtr<FStructDeclaration>& Ptr = DeclaredStructs[Id.Idx];
 	Ptr->RefCount--;
 	if (Ptr->RefCount == 0)
 	{
@@ -178,7 +178,7 @@ void FDeclarations::DropStructRef(FStructId DeclId)
 }
 
 const FStructDeclaration*
-FDeclarations::Find(FStructId Id) const
+FDeclarations::Find(FDeclId Id) const
 {
 	return Id.Idx < uint32(DeclaredStructs.Num()) ? DeclaredStructs[Id.Idx].Get() : nullptr;
 }
@@ -189,7 +189,7 @@ void FDeclarations::Check(FEnumId Id) const
 	checkf(Id.Idx < (uint32)DeclaredEnums.Num() && DeclaredEnums[Id.Idx], TEXT("'%s' is undeclared"), *Debug.Print(Id));
 }
 
-void FDeclarations::Check(FStructId Id) const
+void FDeclarations::Check(FDeclId Id) const
 {
 	checkf(Id.Idx < (uint32)DeclaredStructs.Num() && DeclaredStructs[Id.Idx], TEXT("'%s' is undeclared"), *Debug.Print(Id));
 }
diff --git a/Engine/Plugins/Experimental/PlainProps/Source/Private/PlainPropsDiff.cpp b/Engine/Plugins/Experimental/PlainProps/Source/Private/PlainPropsDiff.cpp
--- a/Engine/Plugins/Experimental/PlainProps/Source/Private/PlainPropsDiff.cpp
+++ b/Engine/Plugins/Experimental/PlainProps/Source/Private/PlainPropsDiff.cpp
@@ -12,7 +12,7 @@ namespace PlainProps
 {
 
 static bool DiffItem(const void* A, const void* B, FDiffContext& Ctx, FRangeMemberBinding Range);
-static bool DiffItem(const void* A, const void* B, FDiffContext& Ctx, FStructId Struct);
+static bool DiffItem(const void* A, const void* B, FDiffContext& Ctx, FBindId Struct);
 
 static bool DiffMember(FLeafMemberBinding Member, const void* OwnerA, const void* OwnerB, FDiffContext& Ctx, FOptionalMemberId Name)
 {
@@ -44,7 +44,7 @@ static bool DiffMember(FRangeMemberBinding Member, const void* OwnerA, const voi
 		const IItemRangeBinding& ItemBinding = Binding.AsItemBinding();
 		if (InnerKind == EMemberKind::Struct)
 		{
-			bDiff = DiffItemRange(A, B, ItemBinding, Ctx, Member.InnermostSchema.Get().AsStruct());
+			bDiff = DiffItemRange(A, B, ItemBinding, Ctx, Member.InnermostSchema.Get().AsStructBindId());
 		}
 		else if (InnerKind == EMemberKind::Leaf)
 		{
@@ -62,15 +62,15 @@ static bool DiffMember(FRangeMemberBinding Member, const void* OwnerA, const voi
 	return bDiff;
 }
 
-static bool DiffMembers(FStructId BindId, const void* A, const void* B, FDiffContext& Ctx);
+static bool DiffMembers(FBindId Id, const void* A, const void* B, FDiffContext& Ctx);
 
-inline bool DiffStruct(FStructId BindId, const void* A, const void* B, FDiffContext& Ctx)
+inline bool DiffStruct(FBindId Id, const void* A, const void* B, FDiffContext& Ctx)
 {
-	if (const ICustomBinding* Custom = Ctx.Customs.FindStruct(BindId))
+	if (const ICustomBinding* Custom = Ctx.Customs.FindStruct(Id))
 	{
 		return Custom->DiffCustom(A, B);
 	}
-	return DiffMembers(BindId, A, B, Ctx);
+	return DiffMembers(Id, A, B, Ctx);
 }
 
 static bool DiffMember(FStructMemberBinding Member, const void* OwnerA, const void* OwnerB, FDiffContext& Ctx, FMemberId Name)
@@ -97,9 +97,9 @@ static bool DiffMember(/* in-out */ FMemberVisitor& It, const void* A, const voi
 }
 
 
-static bool DiffMembers(FStructId BindId, const void* A, const void* B, FDiffContext& Ctx)
+static bool DiffMembers(FBindId Id, const void* A, const void* B, FDiffContext& Ctx)
 {
-	const FSchemaBinding& Schema = Ctx.Schemas.GetStruct(BindId);
+	const FSchemaBinding& Schema = Ctx.Schemas.GetStruct(Id);
 	const FStructDeclaration& Declaration = Ctx.Types.Get(Schema.DeclId);
 	
 	FMemberVisitor It(Schema);
@@ -128,21 +128,21 @@ static bool DiffItem(const void* A, const void* B, FDiffContext& Ctx, FRangeMemb
 	return DiffMember(Range, A, B, Ctx, NoId);
 }
 
-static bool DiffItem(const void* A, const void* B, FDiffContext& Ctx, FStructId Id)
+static bool DiffItem(const void* A, const void* B, FDiffContext& Ctx, FBindId Id)
 {
 	return DiffStruct(Id, A, B, Ctx);
 }
 
 ////////////////////////////////////////////////////////////////////////////////////////////////
 
-bool DiffStructs(FStructId BindId, const void* A, const void* B, FDiffContext& Ctx)
+bool DiffStructs(FBindId Id, const void* A, const void* B, FDiffContext& Ctx)
 {
-	if (const ICustomBinding* Custom = Ctx.Customs.FindStruct(BindId))
+	if (const ICustomBinding* Custom = Ctx.Customs.FindStruct(Id))
 	{
 		// Currently no way for DiffCustom to emit deeper diff into
 		return Custom->DiffCustom(A, B);
 	}
-	return DiffMembers(BindId, A, B, Ctx);
+	return DiffMembers(Id, A, B, Ctx);
 }
 
 ////////////////////////////////////////////////////////////////////////////////
diff --git a/Engine/Plugins/Experimental/PlainProps/Source/Private/PlainPropsInternalBind.h b/Engine/Plugins/Experimental/PlainProps/Source/Private/PlainPropsInternalBind.h
--- a/Engine/Plugins/Experimental/PlainProps/Source/Private/PlainPropsInternalBind.h
+++ b/Engine/Plugins/Experimental/PlainProps/Source/Private/PlainPropsInternalBind.h
@@ -23,7 +23,7 @@ public:
 	FLeafMemberBinding			GrabLeaf();				// @pre PeekKind() == EMemberKind::Leaf
 	FRangeMemberBinding			GrabRange();			// @pre PeekKind() == EMemberKind::Range
 	FStructMemberBinding		GrabStruct();			// @pre PeekKind() == EMemberKind::Struct
-	FStructId					GrabSuper();			// @pre First grab and has declared super
+	FBindId						GrabSuper();			// @pre First grab and has declared super
 	void						SkipMember();
 
 protected: // for unit tests
@@ -38,7 +38,7 @@ protected: // for unit tests
 	uint64						GrabMemberOffset();
 	FMemberBindTypeRange		GrabInnerTypes();
 	FInnerId					GrabInnerSchema();
-	FStructId					GrabStructSchema(FStructType Type);
+	FBindId						GrabStructSchema(FStructType Type);
 	FOptionalInnerId			GrabRangeSchema(FMemberType InnermostType);
 	FEnumId						GrabEnumSchema()		{ return GrabInnerSchema().AsEnum(); }
 };
diff --git a/Engine/Plugins/Experimental/PlainProps/Source/Private/PlainPropsInternalParse.h b/Engine/Plugins/Experimental/PlainProps/Source/Private/PlainPropsInternalParse.h
--- a/Engine/Plugins/Experimental/PlainProps/Source/Private/PlainPropsInternalParse.h
+++ b/Engine/Plugins/Experimental/PlainProps/Source/Private/PlainPropsInternalParse.h
@@ -71,10 +71,10 @@ struct FParsedMemberSchema
 
 struct FParsedStructSchema
 {
-	FStructId						Id;
+	FDeclId							Id;
 	uint16							Version = 0;
-	FOptionalStructId				DeclaredSuper;
-	FOptionalStructId				UsedSuper;
+	FOptionalDeclId					DeclaredSuper;
+	FOptionalDeclId					UsedSuper;
 	TConstArrayView<FMemberId>		MemberNames;	// FStructDeclaration.MemberOrder
 	TArray<FParsedMemberSchema>		MemberSchemas;	// Same size as declared member names
 };
@@ -131,12 +131,12 @@ public:
 		return NoId;
 	}
 
-	FOptionalStructId GetStructId(FType Type) const
+	FOptionalDeclId GetStructId(FType Type) const
 	{
 		FSetElementId Id = Structs.FindId(Type);
 		if (Id.IsValidId())
 		{
-			return FStructId{ IntCastChecked<uint32>(Id.AsInteger()) };
+			return FDeclId{ IntCastChecked<uint32>(Id.AsInteger()) };
 		}
 		return NoId;
 	}
diff --git a/Engine/Plugins/Experimental/PlainProps/Source/Private/PlainPropsLoad.cpp b/Engine/Plugins/Experimental/PlainProps/Source/Private/PlainPropsLoad.cpp
--- a/Engine/Plugins/Experimental/PlainProps/Source/Private/PlainPropsLoad.cpp
+++ b/Engine/Plugins/Experimental/PlainProps/Source/Private/PlainPropsLoad.cpp
@@ -178,38 +178,38 @@ struct FLoadIdMapping
 {
 	FStructSchemaId ReadId; // ~ Batch decl id, index into saved schemas and load plans
 	FStructSchemaId LoadId; // ~ Batch bind id, index into load plans
-	FStructId		BindId; // Runtime bind id
+	FBindId			Id; // Runtime bind id
 };
 
 // Helps load type-erased/lowered structs with ExplicitBindName by allocating new load-time struct ids
 class FLoadIdBinder
 {
 public:
-	FLoadIdBinder(TConstArrayView<FStructId> RuntimeDeclIds, FDebugIds Dbg)
-	: DeclIds(RuntimeDeclIds)
-	, NextLoadIdx(static_cast<uint32>(RuntimeDeclIds.Num()))
+	FLoadIdBinder(TConstArrayView<FDeclId> Ids, FDebugIds Dbg)
+	: DeclIds(Ids)
+	, NextLoadIdx(static_cast<uint32>(Ids.Num()))
 	, Debug(Dbg)
 	{}
 
-	FStructSchemaId					BindLoadId(FStructSchemaId ReadId, FStructId BindId)
+	FStructSchemaId					BindLoadId(FStructSchemaId ReadId, FBindId Id)
 	{
 		FStructId DeclId = DeclIds[ReadId.Idx];
-		FStructSchemaId LoadId = BindId == DeclId ? ReadId : MapLoadId(ReadId, BindId);
+		FStructSchemaId LoadId = Id == DeclId ? ReadId : MapLoadId(ReadId, Id);
 		return LoadId;
 	}
 	FLoadIdMapping					GetMapping(int32 Idx) const { return Mappings[Idx]; }
 	int32							NumMappings() const { return Mappings.Num(); }
 
 private:
-	TConstArrayView<FStructId>			DeclIds;
+	TConstArrayView<FDeclId>			DeclIds;
 	uint32								NextLoadIdx;
 	TArray<FLoadIdMapping>				Mappings;
 
-	FStructSchemaId MapLoadId(FStructSchemaId ReadId, FStructId BindId)
+	FStructSchemaId MapLoadId(FStructSchemaId ReadId, FBindId Id)
 	{
 		for (FLoadIdMapping Mapping : Mappings)
 		{
-			if (Mapping.BindId == BindId)
+			if (Mapping.Id == Id)
 			{
 				check(Mapping.ReadId == ReadId);
 				return Mapping.LoadId;
@@ -219,7 +219,7 @@ private:
 		FLoadIdMapping& Mapping = Mappings.AddDefaulted_GetRef();
 		Mapping.ReadId = ReadId;
 		Mapping.LoadId = { NextLoadIdx++ };
-		Mapping.BindId = BindId;
+		Mapping.Id = Id;
 		return Mapping.LoadId;
 	}
 public:
@@ -229,7 +229,7 @@ public:
 struct FLoadIdBinderDummy
 {
 	FDebugIds							Debug;
-	static constexpr FStructSchemaId	BindLoadId(FStructSchemaId ReadId, FStructId BindId) { return ReadId; }
+	static constexpr FStructSchemaId	BindLoadId(FStructSchemaId ReadId, FBindId Id) { return ReadId; }
 	static constexpr int32				NumMappings() { return 0; }
 	static inline FLoadIdMapping		GetMapping(int32 Idx) { unimplemented(); return {}; }
 };
@@ -315,7 +315,7 @@ void CopyMemberBinding(FRangeMemberBinding Binding, /* in-out */ const FSchemaId
 	if (InnermostType.IsStruct())
 	{
 		FStructSchemaId ReadId = static_cast<FStructSchemaId>(*InnerSchemaIt);
-		FStructSchemaId LoadId = LoadIds.BindLoadId(ReadId, Binding.InnermostSchema.Get().AsStruct());
+		FStructSchemaId LoadId = LoadIds.BindLoadId(ReadId, Binding.InnermostSchema.Get().AsStructBindId());
 		Out.AddInnerSchema(LoadId);
 		++InnerSchemaIt;
 	}
@@ -361,7 +361,7 @@ static void CreateSubsetBindingWithoutEnumIds(const FStructSchema& From, const F
 	// ...first the unnamed super member...
 	if (From.Inheritance != ESuper::No)
 	{
-		FStructId BindId = ToIt.GrabSuper();
+		FBindId BindId = ToIt.GrabSuper();
 		if (UsesSuper(From.Inheritance))
 		{
 			FStructType FromType = From.Footer[0].AsStruct(); // To.Members[0].AsStruct().IsDynamic isn't set
@@ -410,7 +410,7 @@ static void CloneBindingWithReplacedStructIds(const FSchemaId* FromIds, const FS
 	FInnerId* InnerIt = const_cast<FInnerId*>(Schema->GetInnerSchemas());
 	for (FInnerId& Inner : MakeArrayView(InnerIt, To.NumInnerSchemas))
 	{
-		FStructId MemcopiedBindId = Inner.AsStruct();
+		FBindId MemcopiedBindId = Inner.AsStructBindId();
 		reinterpret_cast<FStructSchemaId&>(Inner) = LoadIds.BindLoadId(*ReadIdIt++, MemcopiedBindId);
 	}
 }
@@ -563,20 +563,21 @@ struct FLoadPlanner
 				FLoadIdMapping Mapping;
 				Mapping.ReadId = { Idx };
 				Mapping.LoadId = { Idx };
-				Mapping.BindId = RuntimeIds[Idx];
+				Mapping.Id = FBindId(RuntimeIds[Idx]);
 				CreatePlan(Mapping, LoadIds);
 			}
 
 			return CreateBatch(LoadIds);
 		}
 	
-		FLoadIdBinder LoadIds(RuntimeIds, Declarations.GetDebug());
+		TConstArrayView<FDeclId> DeclIds(static_cast<const FDeclId*>(RuntimeIds.GetData()), RuntimeIds.Num());
+		FLoadIdBinder LoadIds(DeclIds, Declarations.GetDebug());
 		for (uint32 Idx = 0; Idx < NumPlans; ++Idx)
 		{
 			FLoadIdMapping Mapping;
 			Mapping.ReadId = { Idx };
 			Mapping.LoadId = { Idx };
-			Mapping.BindId = RuntimeIds[Idx];
+			Mapping.Id = UpCast(DeclIds[Idx]);
 			CreatePlan(Mapping, LoadIds);
 		}
 
@@ -740,7 +741,7 @@ private:
 		{
 			// Todo: Rename FReadBatchId ReadId everywhere
 			FStructSchemaId ReadId_ = SchemaIt.GrabMemberStruct(Inner.Name);
-			*OutLoadIdIt++ = LoadIds.BindLoadId(ReadId_, Inner.BindId);
+			*OutLoadIdIt++ = LoadIds.BindLoadId(ReadId_, Inner.Id);
 		}
 		check((void*)OutLoadIdIt == SubsetSchemaData.GetData() + SubsetSchemaData.Num());
 
@@ -795,12 +796,12 @@ private:
 		TConstArrayView<FInnerStruct> LoweredInners;
 		if constexpr (std::is_same_v<LoadIdBinder, FLoadIdBinderDummy>)
 		{
-			if (const ICustomBinding* Custom = Customs.FindStruct(Mapping.BindId))
+			if (const ICustomBinding* Custom = Customs.FindStruct(Mapping.Id))
 			{
 				return FLoadStructPlan(*Custom);
 			}
 		}
-		else if (const ICustomBinding* Custom = Customs.FindStruct(Mapping.BindId, /* out */ LoweredInners))
+		else if (const ICustomBinding* Custom = Customs.FindStruct(Mapping.Id, /* out */ LoweredInners))
 		{
 			TArray<FInnerStruct, TInlineAllocator<8>> TmpSubset;
 			TConstArrayView<FInnerStruct> LoweredMembers = GetLoweredMembers(LoweredInners, From.GetMemberNames(), /* out */ TmpSubset);
@@ -815,7 +816,7 @@ private:
 
 		if (From.NumMembers)
 		{
-			if (const FSchemaBinding* To = Schemas.FindStruct(Mapping.BindId))
+			if (const FSchemaBinding* To = Schemas.FindStruct(Mapping.Id))
 			{
 				// Possible optimization - some simple memcpy cases doesn't need to resolve the declaration
 				TConstArrayView<FMemberId> ToMemberIds = Declarations.Get(To->DeclId).GetMemberOrder();
diff --git a/Engine/Plugins/Experimental/PlainProps/Source/Private/PlainPropsParse.cpp b/Engine/Plugins/Experimental/PlainProps/Source/Private/PlainPropsParse.cpp
--- a/Engine/Plugins/Experimental/PlainProps/Source/Private/PlainPropsParse.cpp
+++ b/Engine/Plugins/Experimental/PlainProps/Source/Private/PlainPropsParse.cpp
@@ -622,7 +622,7 @@ public:
 		FDeclarations& InTypes,
 		FParsedSchemas& InSchemas);
 
-	FBuiltStruct* ParseMembers(FStructId SchemaId);
+	FBuiltStruct* ParseMembers(FDeclId SchemaId);
 
 private:
 	FDebugIds GetDebug() const { return FDebugIds(Names); }
@@ -634,9 +634,9 @@ private:
 	bool ParseRange();
 
 	bool AdvanceToNextMember();
-	FStructId ParseDynamicStructType();
+	FDeclId ParseDynamicStructType();
 
-	[[nodiscard]] FBuiltStruct* BuildAndReset(FStructId Schema);
+	[[nodiscard]] FBuiltStruct* BuildAndReset(FDeclId Schema);
 
 	[[nodiscard]] FTypedRange	ParseLeaves(const FMemberSchema& RangeSchema, FScopedTokenReader RangeTokens);
 	[[nodiscard]] FTypedRange	ParseStructs(const FMemberSchema& RangeSchema, FScopedTokenReader RangeTokens);
@@ -646,7 +646,7 @@ private:
 
 	[[nodiscard]] FBuiltStruct*	ParseMembersInternal(
 		FStructType StructType,
-		FStructId Schema,
+		FDeclId Id,
 		FScopedTokenReader InnerTokens);
 
 	[[nodiscard]] FTypedRange	ParseUnicodeLeafValueAsRange(
@@ -662,7 +662,7 @@ private:
 	FTextIndexer&							Names;
 	FDeclarations&							Types;
 	FParsedSchemas&							Schemas;
-	FOptionalStructId						UsedSuper;
+	FOptionalDeclId							UsedSuper;
 	uint32									TokenIdx = 0;
 	uint32								 	MemberIdx = 0;
 	bool									HasSuper = false;
@@ -1169,10 +1169,10 @@ FParsedStructSchema FStructSchemaParser::DeclareStructSchema(
 	TConstArrayView<FUtf8StringView> MemberTypes)
 {
 	const FType Type = MakeType(TypeTokens, Names);
-	const FStructId Id = Names.IndexStruct(Type);
-	const FOptionalStructId DeclaredSuper = ParsedSuper ? ToOptional(Names.IndexStruct(MakeType(ParsedSuper.GetValue(), Names))) : NoId;
+	const FDeclId Id = Names.IndexDeclId(Type);
+	const FOptionalDeclId DeclaredSuper = ParsedSuper ? ToOptional(Names.IndexDeclId(MakeType(ParsedSuper.GetValue(), Names))) : NoId;
 
-	FOptionalStructId UsedSuper;
+	FOptionalDeclId UsedSuper;
 
 	TArray<FMemberId, TInlineAllocator<64>> MemberIds;
 	TArray<FParsedMemberSchema> MemberSchemas;
@@ -1187,7 +1187,7 @@ FParsedStructSchema FStructSchemaParser::DeclareStructSchema(
 			TOptional<FTypeTokens> ParsedSuperMember = TokenizeType(MemberTypes[0]);
 			check(ParsedSuperMember);
 			FType UsedSuperType = MakeType(ParsedSuperMember.GetValue(), Names);
-			UsedSuper = Names.IndexStruct(UsedSuperType);
+			UsedSuper = Names.IndexDeclId(UsedSuperType);
 		}
 	}
 
@@ -1311,7 +1311,7 @@ FMemberParser::FMemberParser(
 {
 }
 
-FBuiltStruct* FMemberParser::ParseMembers(FStructId Id)
+FBuiltStruct* FMemberParser::ParseMembers(FDeclId Id)
 {
 	ParseAll();
 	return BuildAndReset(Id);
@@ -1497,7 +1497,7 @@ bool FMemberParser::ParseStruct()
 
 	FToken Token = MemberTokens.GrabToken(EToken::BeginStruct);
 
-	FStructId Id = StructType.IsDynamic ? ParseDynamicStructType() : MemberSchema.InnerSchema.Get().AsStruct();
+	FDeclId Id = StructType.IsDynamic ? ParseDynamicStructType() : MemberSchema.InnerSchema.Get().AsStructDeclId();
 
 	FBuiltStruct* Struct = ParseMembersInternal(StructType, Id, { MemberTokens.GetTokens(), Token });
 	MemberBuilder.AddStruct(MemberId, Id, Struct);
@@ -1566,7 +1566,7 @@ bool FMemberParser::AdvanceToNextMember()
 	return false;
 }
 
-FStructId FMemberParser::ParseDynamicStructType()
+FDeclId FMemberParser::ParseDynamicStructType()
 {
 	FToken Token = MemberTokens.GrabToken(EToken::LeafId);
 	check(Token.Value() == GLiterals.Dynamic);
@@ -1576,12 +1576,12 @@ FStructId FMemberParser::ParseDynamicStructType()
 	TOptional<FTypeTokens> TypeTokens = TokenizeType(StructName);
 	check(TypeTokens);
 	FType Type = MakeType(TypeTokens.GetValue(), Names);
-	FOptionalStructId Out = Names.GetStructId(Type);
+	FOptionalDeclId Out = Names.GetStructId(Type);
 	checkf(Out, TEXT("Failed to parse dynamic struct schema '%s'"), *Print(Token));
 	return Out.Get();
 }
 
-FBuiltStruct* FMemberParser::BuildAndReset(FStructId Id)
+FBuiltStruct* FMemberParser::BuildAndReset(FDeclId Id)
 {
 	return MemberBuilder.BuildAndReset(Scratch, Types.Get(Id), GetDebug());
 }
@@ -1699,7 +1699,7 @@ FTypedRange FMemberParser::ParseStructs(const FMemberSchema& RangeSchema, FScope
 	check(RangeSchema.InnerSchema || StructType.IsDynamic);
 
 	TArray<FBuiltStruct*, TInlineAllocator64<64>> Structs;
-	FOptionalStructId Schema = ToOptionalStruct(RangeSchema.InnerSchema);
+	FOptionalDeclId Schema = ToOptionalDeclId(RangeSchema.InnerSchema);
 
 	if (StructType.IsDynamic)
 	{
@@ -1707,7 +1707,7 @@ FTypedRange FMemberParser::ParseStructs(const FMemberSchema& RangeSchema, FScope
 		while (RangeTokens.HasMore())
 		{
 			FToken Token = RangeTokens.GrabToken(EToken::BeginStruct);
-			FStructId DynamicSchema = ParseDynamicStructType();
+			FDeclId DynamicSchema = ParseDynamicStructType();
 			checkf(!Schema || Schema == DynamicSchema, TEXT("Heterogeneous struct ranges have not been implemented yet"));
 			Schema = DynamicSchema;
 
@@ -1826,12 +1826,12 @@ FTypedRange FMemberParser::ParseRangeInternal(const FMemberSchema& RangeSchema)
 
 FBuiltStruct* FMemberParser::ParseMembersInternal(
 	FStructType StructType,
-	FStructId Schema,
+	FDeclId Id,
 	FScopedTokenReader InnerTokens)
 {
-	const FParsedStructSchema& ParsedSchema = ResolveStructSchema(Schemas, Schema);
+	const FParsedStructSchema& ParsedSchema = ResolveStructSchema(Schemas, Id);
 	FMemberParser MemberParser(Scratch, ParsedSchema, InnerTokens, Names, Types, Schemas);
-	return MemberParser.ParseMembers(Schema);
+	return MemberParser.ParseMembers(Id);
 }
 
 static FUtf8StringView ParseString(TArray<UTF8CHAR>& Out, FUtf8StringView String)
@@ -2001,7 +2001,7 @@ void FBatchParser::ParseObjects(TConstArrayView64<FToken> TokensView)
 		}
 
 		const FType Type = MakeType(TypeTokens.GetValue(), Names);
-		const FOptionalStructId StructId = Names.GetStructId(Type);
+		const FOptionalDeclId StructId = Names.GetStructId(Type);
 
 		checkf(StructId, TEXT("Object struct schema '%s' not found"), *Print(Token));
 		const FParsedStructSchema& ParsedSchema = ResolveStructSchema(Schemas, StructId.Get());
@@ -2018,13 +2018,12 @@ void FBatchParser::ParseObjects(TConstArrayView64<FToken> TokensView)
 	}
 }
 
-class FParseBindings : public IStructBindIds
+class FParseBindings : public IBindIds
 {
-public:
-	virtual FStructId GetDeclId(FStructId BindId) const override final
+	virtual FDeclId	Lower(FBindId Id) const override final
 	{
 		checkf(false, TEXT("All struct ids should be declared, nothing is bound with different names for text"));
-		return BindId;
+		return LowerCast(Id);
 	}
 };
 
diff --git a/Engine/Plugins/Experimental/PlainProps/Source/Private/PlainPropsSave.cpp b/Engine/Plugins/Experimental/PlainProps/Source/Private/PlainPropsSave.cpp
--- a/Engine/Plugins/Experimental/PlainProps/Source/Private/PlainPropsSave.cpp
+++ b/Engine/Plugins/Experimental/PlainProps/Source/Private/PlainPropsSave.cpp
@@ -124,10 +124,14 @@ struct TNonLeafRangeSaver
 	}
 };
 
-struct FDefaultStruct { FStructId Id; const void* Struct; };
+struct FDefaultStruct
+{
+	FBindId			Id;
+	const void*		Struct;
+};
 
 using FInternalNestedRangeSaver = TNonLeafRangeSaver<FBuiltRange*, FRangeMemberBinding>;
-using FInternalStructRangeSaver = TNonLeafRangeSaver<FBuiltStruct*, FStructId>;
+using FInternalStructRangeSaver = TNonLeafRangeSaver<FBuiltStruct*, FBindId>;
 using FInternalStructRangeDeltaSaver = TNonLeafRangeSaver<FBuiltStruct*, FDefaultStruct>;
 
 //////////////////////////////////////////////////////////////////////////
@@ -199,10 +203,10 @@ FBuiltRange* SaveLeafRange(const void* Range, const ILeafRangeBinding& Binding,
 	return Allocator.GetAllocatedRange();
 }
 
-inline const FStructDeclaration& GetDeclaration(const FSaveContext& Ctx, FStructId BindId)
+inline const FStructDeclaration& GetDeclaration(const FSaveContext& Ctx, FBindId BindId)
 {
 	// BindId is usually DeclId, unless it's been type-erased/lowered
-	if (const FStructDeclaration* Out = Ctx.Declarations.Find(BindId))
+	if (const FStructDeclaration* Out = Ctx.Declarations.Find(LowerCast(BindId)))
 	{
 		return *Out;
 	}
@@ -215,16 +219,16 @@ inline const FStructDeclaration& GetDeclaration(const FSaveContext& Ctx, FStruct
 	return Ctx.Declarations.Get(Ctx.Schemas.GetStruct(BindId).DeclId);
 }
 
-[[nodiscard]] static FBuiltRange* SaveStructRange(const void* Range, const IItemRangeBinding& ItemBinding, const FSaveContext& Ctx, FStructId BindId)
+[[nodiscard]] static FBuiltRange* SaveStructRange(const void* Range, const IItemRangeBinding& ItemBinding, const FSaveContext& Ctx, FBindId Id)
 {
-	if (Ctx.Defaults && GetDeclaration(Ctx, BindId).Occupancy == EMemberPresence::AllowSparse)
+	if (Ctx.Defaults && GetDeclaration(Ctx, Id).Occupancy == EMemberPresence::AllowSparse)
 	{
-		FDefaultStruct Default = { BindId, Ctx.Defaults->Get(BindId) };
+		FDefaultStruct Default = { Id, Ctx.Defaults->Get(Id) };
 		return SaveNonLeafRange<FInternalStructRangeDeltaSaver>(Range, ItemBinding, Ctx, Default);
 	}
 	else
 	{
-		return SaveNonLeafRange<FInternalStructRangeSaver>(Range, ItemBinding, Ctx, BindId);
+		return SaveNonLeafRange<FInternalStructRangeSaver>(Range, ItemBinding, Ctx, Id);
 	}
 }
 
@@ -243,7 +247,7 @@ inline const FStructDeclaration& GetDeclaration(const FSaveContext& Ctx, FStruct
 	{
 	case EMemberKind::Leaf:		return SaveLeafRange(Range, ItemBinding, Ctx, GetItemWidth(InnerType.AsLeaf()));
 	case EMemberKind::Range:	return SaveNonLeafRange<FInternalNestedRangeSaver>(Range, ItemBinding, Ctx, GetInnerRange(Member));
-	case EMemberKind::Struct:	return SaveStructRange(Range, ItemBinding, Ctx, Member.InnermostSchema.Get().AsStruct());
+	case EMemberKind::Struct:	return SaveStructRange(Range, ItemBinding, Ctx, Member.InnermostSchema.Get().AsStructBindId());
 	}
 
 	check(false);
@@ -257,7 +261,7 @@ inline const FStructDeclaration& GetDeclaration(const FSaveContext& Ctx, FStruct
 	return SaveRange(Range, Member, Ctx);
 }
 
-[[nodiscard]] static FBuiltStruct* SaveRangeItem(const uint8* Struct, FStructId Id, const FSaveContext& Ctx)
+[[nodiscard]] static FBuiltStruct* SaveRangeItem(const uint8* Struct, FBindId Id, const FSaveContext& Ctx)
 {
 	return SaveStruct(Struct, Id, Ctx);
 }
@@ -319,10 +323,10 @@ inline void SaveAllMembers(FMemberBuilder& Out, const void* Struct, const FSchem
 	FMemberVisitor It(Schema);
 	if (Declaration.Super)
 	{
-		FStructId SuperId = It.GrabSuper();
-		checkSlow(SuperId == Declaration.Super);
+		FBindId SuperId = It.GrabSuper();
+		checkSlow(SuperId == ToOptionalStruct(Declaration.Super));
 		const FSchemaBinding& SuperSchema = Ctx.Schemas.GetStruct(SuperId);
-		const FStructDeclaration& SuperDecl = Ctx.Declarations.Get(SuperId);
+		const FStructDeclaration& SuperDecl = Ctx.Declarations.Get(SuperSchema.DeclId);
 		SaveAllMembers(Out, Struct, SuperSchema, SuperDecl, Ctx);
 		Out.BuildSuperStruct(Ctx.Scratch, SuperDecl, Ctx.Declarations.GetDebug());
 	}
@@ -339,17 +343,17 @@ inline void SaveAllMembers(FMemberBuilder& Out, const void* Struct, const FSchem
 	checkSlow(!It.HasMore());
 }
 
-FBuiltStruct* SaveStruct(const void* Struct, FStructId BindId, const FSaveContext& Ctx)
+FBuiltStruct* SaveStruct(const void* Struct, FBindId Id, const FSaveContext& Ctx)
 {
 	const FStructDeclaration* Declaration = nullptr;
 	FMemberBuilder Out;
-	if (ICustomBinding* Custom = Ctx.Customs.FindStructToSave(BindId, /* out */ Declaration))
+	if (ICustomBinding* Custom = Ctx.Customs.FindStructToSave(Id, /* out */ Declaration))
 	{
 		Custom->SaveCustom(Out, Struct, nullptr, Ctx);
 	}
 	else
 	{
-		const FSchemaBinding& Schema = Ctx.Schemas.GetStruct(BindId);
+		const FSchemaBinding& Schema = Ctx.Schemas.GetStruct(Id);
 		Declaration = &Ctx.Declarations.Get(Schema.DeclId);	
 		SaveAllMembers(Out, Struct, Schema, *Declaration, Ctx);
 	}
@@ -374,14 +378,14 @@ static bool DiffItem(const uint8* A, const uint8* B, const FSaveContext& Ctx, FR
 	{
 	case EMemberKind::Leaf:		return DiffItemRange(A, B, ItemBinding, Ctx, SizeOf(GetItemWidth(InnerType.AsLeaf())));
 	case EMemberKind::Range:	return DiffItemRange(A, B, ItemBinding, Ctx, GetInnerRange(Member));
-	case EMemberKind::Struct:	return DiffItemRange(A, B, ItemBinding, Ctx, Member.InnermostSchema.Get().AsStruct());
+	case EMemberKind::Struct:	return DiffItemRange(A, B, ItemBinding, Ctx, Member.InnermostSchema.Get().AsStructBindId());
 	}
 
 	check(false);
 	return false;
 }
 
-static bool DiffItem(const uint8* A, const uint8* B, const FSaveContext& Ctx, FStructId Id)
+static bool DiffItem(const uint8* A, const uint8* B, const FSaveContext& Ctx, FBindId Id)
 {
 	if (const ICustomBinding* Custom = Ctx.Customs.FindStruct(Id))
 	{
@@ -442,8 +446,8 @@ static const FStructDeclaration& SaveSchemaBoundStructDelta(FMemberBuilder& Out,
 		FMemberVisitor It(Schema);
 		if (Declaration.Super)
 		{
-			FStructId SuperId = It.GrabSuper();
-			checkSlow(SuperId == Declaration.Super);
+			FBindId SuperId = It.GrabSuper();
+			checkSlow(SuperId == ToOptionalStruct(Declaration.Super));
 			const FSchemaBinding& SuperSchema = Ctx.Schemas.GetStruct(SuperId);
 			const FStructDeclaration& SuperDecl = SaveSchemaBoundStructDelta(Out, Struct, Default, SuperSchema, Ctx);
 			Out.BuildSuperStruct(Ctx.Scratch, SuperDecl, Ctx.Declarations.GetDebug());
@@ -468,11 +472,11 @@ static const FStructDeclaration& SaveSchemaBoundStructDelta(FMemberBuilder& Out,
 	return Declaration;
 }
 
-FBuiltStruct* SaveStructDelta(const void* Struct, const void* Default, FStructId BindId, const FSaveContext& Ctx)
+FBuiltStruct* SaveStructDelta(const void* Struct, const void* Default, FBindId Id, const FSaveContext& Ctx)
 {
 	FMemberBuilder Out;
 	const FStructDeclaration* Declaration = nullptr;
-	if (ICustomBinding* Custom = Ctx.Customs.FindStructToSave(BindId, /* out */ Declaration))
+	if (ICustomBinding* Custom = Ctx.Customs.FindStructToSave(Id, /* out */ Declaration))
 	{
 		if (Custom->DiffCustom(Struct, Default))
 		{
@@ -481,17 +485,17 @@ FBuiltStruct* SaveStructDelta(const void* Struct, const void* Default, FStructId
 	}
 	else
 	{
-		const FSchemaBinding& Schema = Ctx.Schemas.GetStruct(BindId);
+		const FSchemaBinding& Schema = Ctx.Schemas.GetStruct(Id);
 		Declaration = &SaveSchemaBoundStructDelta(Out, Struct, Default, Schema, Ctx);
 	}
 	return Out.BuildAndReset(Ctx.Scratch, *Declaration, Ctx.Declarations.GetDebug());
 }
 
-FBuiltStruct* SaveStructDeltaIfDiff(const void* Struct, const void* Default, FStructId BindId, const FSaveContext& Ctx)
+FBuiltStruct* SaveStructDeltaIfDiff(const void* Struct, const void* Default, FBindId Id, const FSaveContext& Ctx)
 {
 	FMemberBuilder Out;
 	const FStructDeclaration* Declaration = nullptr;
-	if (ICustomBinding* Custom = Ctx.Customs.FindStructToSave(BindId, /* out */ Declaration))
+	if (ICustomBinding* Custom = Ctx.Customs.FindStructToSave(Id, /* out */ Declaration))
 	{
 		if (Custom->DiffCustom(Struct, Default))
 		{
@@ -502,7 +506,7 @@ FBuiltStruct* SaveStructDeltaIfDiff(const void* Struct, const void* Default, FSt
 		return nullptr;
 	}
 
-	const FSchemaBinding& Schema = Ctx.Schemas.GetStruct(BindId);
+	const FSchemaBinding& Schema = Ctx.Schemas.GetStruct(Id);
 	Declaration = &SaveSchemaBoundStructDelta(Out, Struct, Default, Schema, Ctx);
 
 	return Out.IsEmpty() ? nullptr : Out.BuildAndReset(Ctx.Scratch, *Declaration, Ctx.Declarations.GetDebug());
diff --git a/Engine/Plugins/Experimental/PlainProps/Source/Private/PlainPropsUeCoreBindings.cpp b/Engine/Plugins/Experimental/PlainProps/Source/Private/PlainPropsUeCoreBindings.cpp
--- a/Engine/Plugins/Experimental/PlainProps/Source/Private/PlainPropsUeCoreBindings.cpp
+++ b/Engine/Plugins/Experimental/PlainProps/Source/Private/PlainPropsUeCoreBindings.cpp
@@ -12,8 +12,8 @@ void FTransformBinding::Save(FMemberBuilder& Dst, const FTransform& Src, const F
 {
 	static_assert(std::is_same_v<decltype(FTransform().GetTranslation().X), double>);
 	
-	const FStructDeclaration& VectorDecl = Ctx.Declarations.Get(VectorId);
-	const FStructDeclaration& QuatDecl = Ctx.Declarations.Get(QuatId);
+	const FStructDeclaration& VectorDecl = Ctx.Declarations.Get(LowerCast(VectorId));
+	const FStructDeclaration& QuatDecl = Ctx.Declarations.Get(LowerCast(QuatId));
 	FDenseMemberBuilder Inner = { Ctx.Scratch, Ctx.Declarations.GetDebug() };
 
 	FVector T = Src.GetTranslation();
diff --git a/Engine/Plugins/Experimental/PlainProps/Source/Private/PlainPropsWrite.cpp b/Engine/Plugins/Experimental/PlainProps/Source/Private/PlainPropsWrite.cpp
--- a/Engine/Plugins/Experimental/PlainProps/Source/Private/PlainPropsWrite.cpp
+++ b/Engine/Plugins/Experimental/PlainProps/Source/Private/PlainPropsWrite.cpp
@@ -18,11 +18,11 @@ namespace PlainProps
 // Rewrite as a more compact data structure once we get a large number of ids
 struct FWriteIds
 {
-	FWriteIds(const FIdIndexerBase& DeclaredIds, const IStructBindIds& BindIds, const FBuiltSchemas& Schemas, ESchemaFormat Format);
+	FWriteIds(const FIdIndexerBase& DeclaredIds, const IBindIds& BindIds, const FBuiltSchemas& Schemas, ESchemaFormat Format);
 
 	bool								HasStableNames() const { return !Names.IsEmpty(); }
 
-	const IStructBindIds&				BindIds;
+	const IBindIds&						BindIds;
 
 	TArray<FOptionalNameId>				Names;
 	TArray<FOptionalNestedScopeId>		NestedScopes;
@@ -222,7 +222,7 @@ static void CopyUsed(TArray<ItemType>& Out, const TArray<IdType>& Kept, IndexerT
 	}
 }
 
-FWriteIds::FWriteIds(const FIdIndexerBase& Ids, const IStructBindIds& InBindIds, const FBuiltSchemas& Schemas, ESchemaFormat Format)
+FWriteIds::FWriteIds(const FIdIndexerBase& Ids, const IBindIds& InBindIds, const FBuiltSchemas& Schemas, ESchemaFormat Format)
 : BindIds(InBindIds)
 , Debug(Ids)
 {
@@ -290,16 +290,17 @@ FWriteIds::FWriteIds(const FIdIndexerBase& Ids, const IStructBindIds& InBindIds,
 	check(HasStableNames() == (Format == ESchemaFormat::StableNames));
 }
 
-FStructSchemaId FWriteIds::RemapStruct(FStructId OldBindId) const
+FStructSchemaId FWriteIds::RemapStruct(FStructId OldId) const
 {
-	if (FOptionalStructSchemaId WriteDeclId = Structs[OldBindId.Idx])
+	if (FOptionalStructSchemaId WriteDeclId = Structs[OldId.Idx])
 	{
 		return WriteDeclId.Get();
 	}
 
 	// Could optimize by caching Structs[OldBindId.Idx] here
-	FStructId OldDeclId = BindIds.GetDeclId(OldBindId);
+	FDeclId OldDeclId = BindIds.Lower(FBindId(OldId));
 	return Structs[OldDeclId.Idx].Get();
+
 }
 
 //////////////////////////////////////////////////////////////////////////
@@ -337,7 +338,8 @@ static FOptionalSchemaId GetStaticInnerSchema(const FMemberSchema& Schema, const
 		}
 		else if (!InnermostType.AsStruct().IsDynamic)
 		{
-			return { NewIds.RemapStruct(InnerSchema.Get().AsStruct()) };
+			//tmptorp
+			return { NewIds.RemapStruct(InnerSchema.Get().AsStructBindId()) };
 		}
 	}
 	return {};
@@ -348,7 +350,7 @@ static TArray<FSchemaId> GetInnerSchemas(const FBuiltStructSchema& Struct, const
 	TArray<FSchemaId> Out;	
 	if (Inheritance != ESuper::No && Inheritance != ESuper::Reused)
 	{
-		Out.Add(NewIds.RemapStruct(Struct.Super.Get()));
+		Out.Add(NewIds.RemapStruct(FBindId(Struct.Super.Get())));
 	}
 
 	for (const FMemberSchema* Schema : Struct.MemberSchemas)
@@ -639,7 +641,7 @@ public:
 
 	FStructSchemaId WriteMembers(FStructId BuiltId, const FBuiltStruct& Struct)
 	{
-		FStructSchemaId WriteId = NewIds.RemapStruct(BuiltId);
+		FStructSchemaId WriteId = NewIds.RemapStruct(FBindId(BuiltId));
 		const FBuiltStructSchema& Schema = Schemas.Structs[WriteId.Idx];
 		const TArray<FMemberId>& Order = Schema.MemberNames;
 		const int32 NumSuper = Schema.MemberSchemas.Num() - Schema.MemberNames.Num();
@@ -702,7 +704,7 @@ private:
 	{
 		if (InnermostType.IsStruct() && InnermostType.AsStruct().IsDynamic)
 		{
-			FSchemaId WriteId = NewIds.RemapStruct(Schema.InnerSchema.Get().AsStruct());
+			FSchemaId WriteId = NewIds.RemapStruct(FBindId(Schema.InnerSchema.Get().AsStruct()));
 			Bytes.Append(reinterpret_cast<uint8*>(&WriteId.Idx), sizeof(WriteId.Idx));
 		}
 
@@ -818,7 +820,7 @@ private:
 
 ////////////////////////////////////////////////////////////////////////////////////////////////////////////////////////////////////////////////////
 
-FWriter::FWriter(const FIdIndexerBase& AllIds, const IStructBindIds& BindIds, const FBuiltSchemas& InSchemas, ESchemaFormat Format)
+FWriter::FWriter(const FIdIndexerBase& AllIds, const IBindIds& BindIds, const FBuiltSchemas& InSchemas, ESchemaFormat Format)
 : Schemas(InSchemas)
 , Debug(AllIds)
 , NewIds(new FWriteIds(AllIds, BindIds, InSchemas, Format))
diff --git a/Engine/Plugins/Experimental/PlainProps/Source/Public/PlainPropsBind.h b/Engine/Plugins/Experimental/PlainProps/Source/Public/PlainPropsBind.h
--- a/Engine/Plugins/Experimental/PlainProps/Source/Public/PlainPropsBind.h
+++ b/Engine/Plugins/Experimental/PlainProps/Source/Public/PlainPropsBind.h
@@ -135,7 +135,7 @@ static_assert(sizeof(FMemberBindType) == 1);
 // Members are loaded in saved FStructSchema order, not current offset order unless upgrade layer reorders
 struct alignas(/*FRangeBinding*/ 8) FSchemaBinding
 {
-	FStructId				DeclId;
+	FDeclId					DeclId;
 	uint16					NumMembers;
 	uint16					NumInnerSchemas;
 	uint16					NumInnerRanges;
@@ -215,7 +215,7 @@ struct FRangeMemberBinding
 struct FStructMemberBinding
 {
 	FStructType				Type;
-	FStructId				Id;
+	FBindId					Id;
 	SIZE_T					Offset;
 };
 
@@ -223,20 +223,32 @@ PLAINPROPS_API FRangeMemberBinding GetInnerRange(FRangeMemberBinding In); // @pr
 
 ////////////////////////////////////////////////////////////////////////////////////////////////
 
-struct FDualStructId
+struct FBothStructId
 {
-	FStructId				BindId;
-	FStructId				DeclId;
+	FBindId					BindId;
+	FDeclId					DeclId;
 
 	bool					IsLowered() const { return BindId != DeclId; }
 };
 
 ////////////////////////////////////////////////////////////////////////////////////////////////
 
+// FStructId statically known to share FBindId and FDeclId, i.e. not lowered
+struct FDualStructId : FStructId
+{
+	explicit FDualStructId(FStructId In) : FStructId(In) {}
+	
+	operator FBindId() const { return {Idx}; }
+	operator FDeclId() const { return {Idx}; }
+};
+
+////////////////////////////////////////////////////////////////////////////////////////////////
+
+// Helps raise lowered FDeclId by seeing which FBindId a particular named member has
 struct FInnerStruct
 {
 	FMemberId				Name; // Of outer Range or Struct
-	FStructId				BindId;
+	FBindId					Id;
 };
 
 enum ECustomLoadMethod { Construct, Assign };
@@ -276,9 +288,9 @@ public:
 	explicit FCustomBindingMap(FDebugIds Dbg) : Debug(Dbg) {}
 	~FCustomBindingMap() { FMemory::Free(Values); }
 
-	void							Bind(FStructId BindId, ICustomBinding& Binding, const FStructDeclaration& Declaration, FInnersHandle LoweredInners);
-	FCustomBindingHandle			Find(FStructId BindId) const;
-	void							Drop(FStructId BindId);
+	void							Bind(FBindId Id, ICustomBinding& Binding, const FStructDeclaration& Declaration, FInnersHandle LoweredInners);
+	FCustomBindingHandle			Find(FBindId Id) const;
+	void							Drop(FBindId Id);
 
 private:
 	FIdWindow						Window;
@@ -293,14 +305,14 @@ class FCustomBindings
 {
 public:
 	// @param Binding must outlive this or call DropStruct()
-	PLAINPROPS_API void							BindStruct(FStructId BindId, ICustomBinding& Binding, const FStructDeclaration& Declaration, TConstArrayView<FInnerStruct> LoweredInners = {});
-	PLAINPROPS_API const ICustomBinding*		FindStruct(FStructId BindId) const;
-	PLAINPROPS_API const ICustomBinding*		FindStruct(FStructId BindId, TConstArrayView<FInnerStruct>& OutLoweredInners) const;
-	PLAINPROPS_API ICustomBinding*				FindStructToSave(FStructId BindId, const FStructDeclaration*& OutDeclaration) const;
-	PLAINPROPS_API const FStructDeclaration*	FindDeclaration(FStructId BindId) const;
-	PLAINPROPS_API void							DropStruct(FStructId BindId);
-
-	virtual FCustomBindingHandle				Find(FStructId BindId) const = 0;
+	PLAINPROPS_API void							BindStruct(FBindId Id, ICustomBinding& Binding, const FStructDeclaration& Declaration, TConstArrayView<FInnerStruct> LoweredInners = {});
+	PLAINPROPS_API const ICustomBinding*		FindStruct(FBindId Id) const;
+	PLAINPROPS_API const ICustomBinding*		FindStruct(FBindId Id, TConstArrayView<FInnerStruct>& OutLoweredInners) const;
+	PLAINPROPS_API ICustomBinding*				FindStructToSave(FBindId Id, const FStructDeclaration*& OutDeclaration) const;
+	PLAINPROPS_API const FStructDeclaration*	FindDeclaration(FBindId Id) const;
+	PLAINPROPS_API void							DropStruct(FBindId Id);
+
+	virtual FCustomBindingHandle				Find(FBindId Id) const = 0;
 protected:
 	FCustomBindings(TArray<FInnerStruct>& Inners, FDebugIds Dbg) : Map(Dbg), BottomInners(Inners) {}
 	FCustomBindings(const FCustomBindings* Under) : Map(Under->Map.Debug), BottomInners(Under->BottomInners) {}
@@ -313,7 +325,7 @@ protected:
 class FCustomBindingsBottom final : public FCustomBindings
 {
 	mutable TArray<FInnerStruct>				AllInners;
-	PLAINPROPS_API virtual FCustomBindingHandle	Find(FStructId BindId) const override;	
+	PLAINPROPS_API virtual FCustomBindingHandle	Find(FBindId Id) const override;	
 public:
 	explicit FCustomBindingsBottom(FDebugIds Dbg) : FCustomBindings(/* unconstructed */  AllInners, Dbg) {}
 };
@@ -321,7 +333,7 @@ public:
 class FCustomBindingsOverlay final : public FCustomBindings
 {
 	const FCustomBindings&						Underlay;
-	PLAINPROPS_API virtual FCustomBindingHandle	Find(FStructId BindId) const override;
+	PLAINPROPS_API virtual FCustomBindingHandle	Find(FBindId Id) const override;
 public:
 	explicit FCustomBindingsOverlay(const FCustomBindings& Under) : FCustomBindings(&Under), Underlay(Under) {}
 };
@@ -577,7 +589,7 @@ using RangeBind = typename TRangeBind<T>::Type;
 
 ////////////////////////////////////////////////////////////////////////////////////////////////
 
-struct FDualType
+struct FBothType
 {
 	FType		BindType;
 	FType		DeclType;
@@ -600,22 +612,22 @@ struct FMemberBinding
 	TConstArrayView<FRangeBinding>	RangeBindings;		// Non-empty -> Range
 
 	// @pre InnermostSchema isn't type-erased / lowered
-	PLAINPROPS_API FDualType		IndexParameterName(FIdIndexerBase& Ids) const;
+	PLAINPROPS_API FBothType		IndexParameterName(FIdIndexerBase& Ids) const;
 };
 
-class FSchemaBindings : public IStructBindIds
+class FSchemaBindings : public IBindIds
 {
 public:
 	UE_NONCOPYABLE(FSchemaBindings);
 	explicit FSchemaBindings(FDebugIds In) : Debug(In) {}
 	PLAINPROPS_API ~FSchemaBindings();
 
-	PLAINPROPS_API void						BindStruct(FStructId BindId, FStructId DeclId, TConstArrayView<FMemberBinding> Schema);
-	PLAINPROPS_API const FSchemaBinding*	FindStruct(FStructId BindId) const;
-	PLAINPROPS_API const FSchemaBinding&	GetStruct(FStructId BindId) const;
-	PLAINPROPS_API void						DropStruct(FStructId BindId);
+	PLAINPROPS_API void						BindStruct(FBindId Id, FDeclId DeclId, TConstArrayView<FMemberBinding> Schema);
+	PLAINPROPS_API const FSchemaBinding*	FindStruct(FBindId Id) const;
+	PLAINPROPS_API const FSchemaBinding&	GetStruct(FBindId Id) const;
+	PLAINPROPS_API void						DropStruct(FBindId Id);
 
-	PLAINPROPS_API virtual FStructId	GetDeclId(FStructId BindId) const override final;
+	PLAINPROPS_API virtual FDeclId			Lower(FBindId Id) const override final;
 private:
 	TArray<TUniquePtr<FSchemaBinding>>		Bindings;
 	FDebugIds								Debug;
@@ -626,12 +638,12 @@ private:
 // Lookup default instances when delta-saving struct ranges
 struct IDefaultStructs
 {
-	virtual const void* Get(FStructId BindId) = 0;
+	virtual const void* Get(FBindId Id) = 0;
 };
 
 ////////////////////////////////////////////////////////////////////////////////////////////////
 
-struct FStructBindIds final : IStructBindIds
+struct FStructBindIds final : IBindIds
 {
 	FStructBindIds(const FCustomBindings& InCustoms, const FSchemaBindings& InSchemas)
 	: Customs(InCustoms)
@@ -641,7 +653,7 @@ struct FStructBindIds final : IStructBindIds
 	const FCustomBindings& Customs;
 	const FSchemaBindings& Schemas;
 
-	PLAINPROPS_API virtual FStructId GetDeclId(FStructId BindId) const override;
+	PLAINPROPS_API virtual FDeclId Lower(FBindId Id) const override;
 };
 
 ////////////////////////////////////////////////////////////////////////////////////////////////
@@ -690,30 +702,30 @@ FType IndexStructName()
 }
 
 template<class Ids, class Typename>
-FStructId IndexStructBindIdIfNeeded(FStructId DeclId)
+FBindId IndexStructBindIdIfNeeded(FDeclId DeclId)
 {
 	if constexpr (ExplicitBindName<Typename> || ParametricName<Typename>)
 	{
 		// Note could pass in and reuse declared namespace here
-		return Ids::IndexStruct(IndexStructName<Ids, ETypename::Bind, Typename>());
+		return FBindId(Ids::IndexStruct(IndexStructName<Ids, ETypename::Bind, Typename>()));
 	}
 	else
 	{
-		return DeclId;
+		return UpCast(DeclId);
 	}
 }
 
 template<class Ids, class Typename>
-FDualStructId IndexStructDualId(FType DeclName = IndexStructName<Ids, ETypename::Decl, Typename>())
+FBothStructId IndexStructBothId(FType DeclName = IndexStructName<Ids, ETypename::Decl, Typename>())
 {	
-	FDualStructId Out;
-	Out.DeclId = Ids::IndexStruct(DeclName);
-	Out.BindId = Out.DeclId;
+	FBothStructId Out;
+	Out.DeclId = FDeclId(Ids::IndexStruct(DeclName));
+	Out.BindId = UpCast(Out.DeclId);
 
 	if constexpr (ExplicitBindName<Typename> || ParametricName<Typename>)
 	{
 		FType BindName = IndexStructName<Ids, ETypename::Bind, Typename>();
-		Out.BindId = BindName != DeclName ? Ids::IndexStruct(BindName) : Out.DeclId;
+		Out.BindId = BindName != DeclName ? FBindId(Ids::IndexStruct(BindName)) : UpCast(Out.DeclId);
 	}
 
 	return Out;
@@ -721,25 +733,25 @@ FDualStructId IndexStructDualId(FType DeclName = IndexStructName<Ids, ETypename:
 
 // Cached by function static
 template<class Ids, typename Struct>
-FStructId GetStructDeclId()
+FDeclId GetStructDeclId()
 {
-	static FStructId Id = Ids::IndexStruct(IndexStructName<Ids, ETypename::Decl, TTypename<Struct>>());
+	static FDeclId Id = FDeclId(Ids::IndexStruct(IndexStructName<Ids, ETypename::Decl, TTypename<Struct>>()));
 	return Id;
 }
 
 // Cached by function static
 template<class Ids, typename Struct>
-FStructId GetStructBindId()
+FBindId GetStructBindId()
 {
-	static FStructId Id = Ids::IndexStruct(IndexStructName<Ids, ETypename::Bind, TTypename<Struct>>());
+	static FBindId Id = FBindId(Ids::IndexStruct(IndexStructName<Ids, ETypename::Bind, TTypename<Struct>>()));
 	return Id;
 }
 
 // Cached by function static
 template<class Ids, typename Struct>
-FDualStructId GetStructDualId()
+FBothStructId GetStructBothId()
 {
-	static FDualStructId Id = IndexStructDualId<Ids, TTypename<Struct>>();
+	static FBothStructId Id = IndexStructBothId<Ids, TTypename<Struct>>();
 	return Id;
 }
 
@@ -831,7 +843,7 @@ struct FCustomInit
 	InnerStructArray& OutLowered;
 
 	// Only needed for inner structs that might get type-erased/lowered
-	void RegisterInnerStruct(FDualStructId Id, TConstArrayView<FMemberId> Names)
+	void RegisterInnerStruct(FBothStructId Id, TConstArrayView<FMemberId> Names)
 	{
 		if (Id.IsLowered())
 		{
@@ -854,25 +866,25 @@ struct TDependent { using StructLoadView = FStructLoadView; };
 
 // Scoped custom binding for native type
 template<class T, class Runtime, typename CustomBinding = CustomBind<T>>
-struct TScopedStructBinding : FDualStructId, CustomBinding
+struct TScopedStructBinding : FBothStructId, CustomBinding
 {
 	using Ids = typename Runtime::Ids;
 	using Typename = TCustomTypename<CustomBinding>::Type;
 
 	explicit TScopedStructBinding(FType DeclName = IndexStructName<Ids, ETypename::Decl, Typename>(), InnerStructArray Lowered = {})
-	: TScopedStructBinding(DeclName, IndexStructDualId<Ids, Typename>(DeclName), Lowered)
+	: TScopedStructBinding(DeclName, IndexStructBothId<Ids, Typename>(DeclName), Lowered)
 	{}
 
-	explicit TScopedStructBinding(FDualStructId Id, InnerStructArray Lowered = {})
+	explicit TScopedStructBinding(FBothStructId Id, InnerStructArray Lowered = {})
 	: TScopedStructBinding(Ids::GetIndexer().Resolve(Id.DeclId), Id, Lowered)
 	{}
 
-	explicit TScopedStructBinding(FStructId Id, InnerStructArray Lowered = {})
-	: TScopedStructBinding(Ids::GetIndexer().Resolve(Id), {Id, Id}, Lowered)
+	explicit TScopedStructBinding(FDualStructId Id, InnerStructArray Lowered = {})
+	: TScopedStructBinding(Ids::GetIndexer().Resolve(Id), FBothStructId{Id, Id}, Lowered)
 	{}
 
-	TScopedStructBinding(FType DeclName, FDualStructId DualId, InnerStructArray& Lowered)
-	: FDualStructId(DualId)
+	TScopedStructBinding(FType DeclName, FBothStructId BothId, InnerStructArray& Lowered)
+	: FBothStructId(BothId)
 	, CustomBinding(TCustomInit<Ids>{{Lowered}})
 	{
 		const FStructDeclaration& Declaration = Runtime::GetTypes().DeclareStruct(DeclId, DeclName, 0, CustomBinding::MemberIds, TOccupancyOf<T>::Value);
@@ -908,8 +920,8 @@ struct TScopedStructBinding<T, Runtime, void>
 	using Ids = typename Runtime::Ids;
 	using Typename = TTypename<T>;
 	
-	const FStructId DeclId;
-	const FStructId BindId;
+	const FDeclId DeclId;
+	const FBindId BindId;
 
 	TScopedStructBinding(FType DeclName = IndexStructName<Ids, ETypename::Decl, Typename>())
 	: DeclId(DeclareNativeStruct<CttiOf<T>, Ids>(Runtime::GetTypes(), TOccupancyOf<T>::Value))
@@ -926,7 +938,7 @@ struct TScopedStructBinding<T, Runtime, void>
 };
 
 template<class Struct, typename CustomBinding, class Runtime>
-FStructId BindCustomStructOnce()
+FBindId BindCustomStructOnce()
 {
 	static TScopedStructBinding<Struct, Runtime, CustomBinding> Instance;
 	return Instance.BindId;
@@ -1105,8 +1117,8 @@ FStructId DeclareNativeStruct(FDeclarations& Out, EMemberPresence Occupancy)
 	using SuperType = typename Ctti::Super;
 
 	FType Type = IndexStructName<Ids, ETypename::Decl, Typename>();
-	FStructId Id = Ids::IndexStruct(Type);
-	FOptionalStructId SuperId;
+	FDeclId Id(Ids::IndexStruct(Type));
+	FOptionalDeclId SuperId;
 	if constexpr (!std::is_void_v<SuperType>)
 	{
 		SuperId = GetStructDeclId<Ids, SuperType>();
@@ -1123,7 +1135,7 @@ FStructId DeclareNativeStruct(FDeclarations& Out, EMemberPresence Occupancy)
 }
 
 template<class Ctti, class Runtime>
-void BindNativeStruct(FSchemaBindings& Out, FStructId BindId, FStructId DeclId)
+void BindNativeStruct(FSchemaBindings& Out, FBindId BindId, FDeclId DeclId)
 {
 	FMemberBinding MemberBindings[Ctti::NumVars];
 	ForEachVar<Ctti>([&]<class Var>()
diff --git a/Engine/Plugins/Experimental/PlainProps/Source/Public/PlainPropsBuildSchema.h b/Engine/Plugins/Experimental/PlainProps/Source/Public/PlainPropsBuildSchema.h
--- a/Engine/Plugins/Experimental/PlainProps/Source/Public/PlainPropsBuildSchema.h
+++ b/Engine/Plugins/Experimental/PlainProps/Source/Public/PlainPropsBuildSchema.h
@@ -47,13 +47,13 @@ public:
 	using FStructDeclarations = TConstArrayView<TUniquePtr<FStructDeclaration>>;
 	using FEnumDeclarations = TConstArrayView<TUniquePtr<FEnumDeclaration>>;
 
-	PLAINPROPS_API FSchemasBuilder(const FDeclarations& Types, const FIds& InIds, const IStructBindIds& InBindIds, FScratchAllocator& InScratch, ESchemaFormat InFormat);
-	PLAINPROPS_API FSchemasBuilder(FStructDeclarations InStructs, FEnumDeclarations InEnums, const FIds& InIds, const IStructBindIds& BindIds, FScratchAllocator& Scratch, ESchemaFormat InFormat);
+	PLAINPROPS_API FSchemasBuilder(const FDeclarations& Types, const FIds& InIds, const IBindIds& InBindIds, FScratchAllocator& InScratch, ESchemaFormat InFormat);
+	PLAINPROPS_API FSchemasBuilder(FStructDeclarations InStructs, FEnumDeclarations InEnums, const FIds& InIds, const IBindIds& BindIds, FScratchAllocator& Scratch, ESchemaFormat InFormat);
 	PLAINPROPS_API ~FSchemasBuilder();
 
 	PLAINPROPS_API FEnumSchemaBuilder&			NoteEnum(FEnumId Id);
-	PLAINPROPS_API FStructSchemaBuilder&		NoteStruct(FStructId BindId);
-	PLAINPROPS_API void							NoteStructAndMembers(FStructId BindId, const FBuiltStruct& Struct);
+	PLAINPROPS_API FStructSchemaBuilder&		NoteStruct(FStructId Id);
+	PLAINPROPS_API void							NoteStructAndMembers(FStructId Id, const FBuiltStruct& Struct);
 	PLAINPROPS_API FBuiltSchemas				Build();
 	
 	FScratchAllocator&							GetScratch() const { return Scratch; }
@@ -66,7 +66,7 @@ private:
 	TSet<FStructId>								StructIndices;
 	TSet<FEnumId>								EnumIndices;
 	const FIds&									Ids;
-	const IStructBindIds&						BindIds;
+	const IBindIds&								BindIds;
 	ESchemaFormat								Format;
 	TPagedArray<FStructSchemaBuilder, 4096>		Structs;		// TPagedArray for stable references
 	TPagedArray<FEnumSchemaBuilder, 4096>		Enums;			// TPagedArray for stable references
diff --git a/Engine/Plugins/Experimental/PlainProps/Source/Public/PlainPropsDeclare.h b/Engine/Plugins/Experimental/PlainProps/Source/Public/PlainPropsDeclare.h
--- a/Engine/Plugins/Experimental/PlainProps/Source/Public/PlainPropsDeclare.h
+++ b/Engine/Plugins/Experimental/PlainProps/Source/Public/PlainPropsDeclare.h
@@ -37,9 +37,9 @@ enum class EMemberPresence : uint8 { RequireAll, AllowSparse };
 struct FStructDeclaration
 {
 	uint32					RefCount;
-	FStructId				Id;				// Could be removed, might allow declaration dedup among templated types
+	FDeclId					Id;				// Could be removed, might allow declaration dedup among templated types
 	FType					Type;			// Could be removed, might allow declaration dedup among templated types
-	FOptionalStructId		Super;
+	FOptionalDeclId			Super;
 	uint16					Version;
 	uint16					NumMembers;
 	EMemberPresence			Occupancy;
@@ -63,17 +63,17 @@ public:
 	UE_NONCOPYABLE(FDeclarations);
 	explicit FDeclarations(FDebugIds In) : Debug(In) {}
 
-	PLAINPROPS_API const FEnumDeclaration&			DeclareEnum(FEnumId DeclId, FType Type, EEnumMode Mode, TConstArrayView<FEnumerator> Enumerators, EEnumAliases Policy);
+	PLAINPROPS_API const FEnumDeclaration&			DeclareEnum(FEnumId Id, FType Type, EEnumMode Mode, TConstArrayView<FEnumerator> Enumerators, EEnumAliases Policy);
 	// Declare struct with ref count 1 or increment it and check that previous declaration matches
-	PLAINPROPS_API const FStructDeclaration&		DeclareStruct(FStructId DeclId, FType Type, uint16 Version, TConstArrayView<FMemberId> MemberOrder, EMemberPresence Occupancy, FOptionalStructId Super = {});
-	PLAINPROPS_API const FStructDeclaration&		DeclareNumeralStruct(FStructId DeclId, FType Type, TConstArrayView<FMemberId> Numerals, EMemberPresence Occupancy);
+	PLAINPROPS_API const FStructDeclaration&		DeclareStruct(FDeclId Id, FType Type, uint16 Version, TConstArrayView<FMemberId> MemberOrder, EMemberPresence Occupancy, FOptionalDeclId Super = {});
+	PLAINPROPS_API const FStructDeclaration&		DeclareNumeralStruct(FDeclId Id, FType Type, TConstArrayView<FMemberId> Numerals, EMemberPresence Occupancy);
 	
-	void											DropEnum(FEnumId DeclId)		{ Check(DeclId); DeclaredEnums[DeclId.Idx].Reset(); }
-	PLAINPROPS_API void								DropStructRef(FStructId DeclId);
+	void											DropEnum(FEnumId Id)			{ Check(Id); DeclaredEnums[Id.Idx].Reset(); }
+	PLAINPROPS_API void								DropStructRef(FDeclId Id);
 
-	const FEnumDeclaration&							Get(FEnumId DeclId) const		{ Check(DeclId); return *DeclaredEnums[DeclId.Idx]; }
-	const FStructDeclaration&						Get(FStructId DeclId) const		{ Check(DeclId); return *DeclaredStructs[DeclId.Idx]; }
-	PLAINPROPS_API const FStructDeclaration*		Find(FStructId Id) const;
+	const FEnumDeclaration&							Get(FEnumId Id) const			{ Check(Id); return *DeclaredEnums[Id.Idx]; }
+	const FStructDeclaration&						Get(FDeclId Id) const			{ Check(Id); return *DeclaredStructs[Id.Idx]; }
+	PLAINPROPS_API const FStructDeclaration*		Find(FDeclId Id) const;
 	
 	TConstArrayView<TUniquePtr<FEnumDeclaration>>	GetEnums() const				{ return DeclaredEnums; }
 	TConstArrayView<TUniquePtr<FStructDeclaration>>	GetStructs() const				{ return DeclaredStructs; }
@@ -86,16 +86,10 @@ protected:
 
 #if DO_CHECK
 	PLAINPROPS_API void								Check(FEnumId Id) const;
-	PLAINPROPS_API void								Check(FStructId Id) const;
+	PLAINPROPS_API void								Check(FDeclId Id) const;
 #else
 	void											Check(...) const {}
 #endif
 };
 
-struct IStructBindIds
-{
-	virtual ~IStructBindIds() = default;
-	virtual FStructId GetDeclId(FStructId BindId) const = 0;
-};
-
 } // namespace PlainProps
diff --git a/Engine/Plugins/Experimental/PlainProps/Source/Public/PlainPropsDiff.h b/Engine/Plugins/Experimental/PlainProps/Source/Public/PlainPropsDiff.h
--- a/Engine/Plugins/Experimental/PlainProps/Source/Public/PlainPropsDiff.h
+++ b/Engine/Plugins/Experimental/PlainProps/Source/Public/PlainPropsDiff.h
@@ -34,7 +34,7 @@ struct FDiffContext
 	FDiffPath				Out;
 };
 
-[[nodiscard]] PLAINPROPS_API bool DiffStructs(FStructId BindId, const void* A, const void* B, FDiffContext& Ctx);
+[[nodiscard]] PLAINPROPS_API bool DiffStructs(FBindId Id, const void* A, const void* B, FDiffContext& Ctx);
 
 ////////////////////////////////////////////////////////////////////////////////
 
diff --git a/Engine/Plugins/Experimental/PlainProps/Source/Public/PlainPropsIndex.h b/Engine/Plugins/Experimental/PlainProps/Source/Public/PlainPropsIndex.h
--- a/Engine/Plugins/Experimental/PlainProps/Source/Public/PlainPropsIndex.h
+++ b/Engine/Plugins/Experimental/PlainProps/Source/Public/PlainPropsIndex.h
@@ -59,6 +59,8 @@ public:
 
 	PLAINPROPS_API FEnumId				IndexEnum(FType Type);
 	PLAINPROPS_API FStructId			IndexStruct(FType Type);
+	FDeclId								IndexDeclId(FType Type)								{ return static_cast<FDeclId>(IndexStruct(Type)); }
+	FBindId								IndexBindId(FType Type)								{ return static_cast<FBindId>(IndexStruct(Type)); }
 	
 	virtual uint32						NumNestedScopes() const override final				{ return IntCastChecked<uint32>(NestedScopes.Num()); }
 	virtual uint32						NumParametricTypes() const override final			{ return IntCastChecked<uint32>(ParametricTypes.Num()); }
diff --git a/Engine/Plugins/Experimental/PlainProps/Source/Public/PlainPropsSave.h b/Engine/Plugins/Experimental/PlainProps/Source/Public/PlainPropsSave.h
--- a/Engine/Plugins/Experimental/PlainProps/Source/Public/PlainPropsSave.h
+++ b/Engine/Plugins/Experimental/PlainProps/Source/Public/PlainPropsSave.h
@@ -34,9 +34,9 @@ FSaveContext MakeSaveContext(FScratchAllocator& Scratch)
 
 ////////////////////////////////////////////////////////////////////////////////////////////////
 
-[[nodiscard]] PLAINPROPS_API FBuiltStruct*	SaveStruct(const void* Struct, FStructId BindId, const FSaveContext& Context);
-[[nodiscard]] PLAINPROPS_API FBuiltStruct*	SaveStructDelta(const void* Struct, const void* Default, FStructId BindId, const FSaveContext& Context);
-[[nodiscard]] PLAINPROPS_API FBuiltStruct*	SaveStructDeltaIfDiff(const void* Struct, const void* Default, FStructId BindId, const FSaveContext& Context);
+[[nodiscard]] PLAINPROPS_API FBuiltStruct*	SaveStruct(const void* Struct, FBindId BindId, const FSaveContext& Context);
+[[nodiscard]] PLAINPROPS_API FBuiltStruct*	SaveStructDelta(const void* Struct, const void* Default, FBindId BindId, const FSaveContext& Context);
+[[nodiscard]] PLAINPROPS_API FBuiltStruct*	SaveStructDeltaIfDiff(const void* Struct, const void* Default, FBindId BindId, const FSaveContext& Context);
 [[nodiscard]] PLAINPROPS_API FBuiltRange*	SaveRange(const void* Range, FRangeMemberBinding Member, const FSaveContext& Ctx);
 [[nodiscard]] PLAINPROPS_API FBuiltRange*	SaveLeafRange(const void* Range, const ILeafRangeBinding& Binding, FUnpackedLeafType Leaf, const FSaveContext& Ctx);
 
diff --git a/Engine/Plugins/Experimental/PlainProps/Source/Public/PlainPropsSaveMember.h b/Engine/Plugins/Experimental/PlainProps/Source/Public/PlainPropsSaveMember.h
--- a/Engine/Plugins/Experimental/PlainProps/Source/Public/PlainPropsSaveMember.h
+++ b/Engine/Plugins/Experimental/PlainProps/Source/Public/PlainPropsSaveMember.h
@@ -18,15 +18,15 @@ struct TStructSerializer
 	static constexpr EMemberKind Kind = EMemberKind::Struct;
 
 	template<class Ids>
-	TStructSerializer(TCustomInit<Ids> Init, TConstArrayView<FMemberId> Names, FDualStructId Dual = GetStructDualId<Ids, T>())
-	: BindId(Dual.BindId)
+	TStructSerializer(TCustomInit<Ids> Init, TConstArrayView<FMemberId> Names, FBothStructId Both = GetStructBothId<Ids, T>())
+	: Id(Both.BindId)
 	{
-		Init.RegisterInnerStruct(Dual, Names);
+		Init.RegisterInnerStruct(Both, Names);
 	}
 
 	FTypedValue SaveMember(const T& Value, const FSaveContext& Ctx) const
 	{
-		FMemberSchema Schema = { DefaultStructType, DefaultStructType, 0, FInnerId(BindId), nullptr};
+		FMemberSchema Schema = { DefaultStructType, DefaultStructType, 0, FInnerId(Id), nullptr};
 		return { Schema, {.Struct = SaveItem(Value, Ctx)} };
 	}
 
@@ -54,15 +54,15 @@ struct TStructSerializer
 
 	FBuiltStruct* SaveItem(const T& Value, const FSaveContext& Ctx) const
 	{
-		return SaveStruct(&Value, BindId, Ctx);
+		return SaveStruct(&Value, Id, Ctx);
 	}
 
 	FMemberSchema MakeMemberRangeSchema(ERangeSizeType MaxSize) const
 	{
-		return { FMemberType(MaxSize), DefaultStructType, 1, FInnerId(BindId), nullptr };
+		return { FMemberType(MaxSize), DefaultStructType, 1, FInnerId(Id), nullptr };
 	}
 
-	const FStructId BindId;
+	const FBindId Id;
 };
 
 //////////////////////////////////////////////////////////////////////////
@@ -123,7 +123,7 @@ struct TRangeSerializer
 
 		if constexpr (!LeafType<InnermostType>)
 		{			
-			FDualStructId Id = {InnermostSchema.Get().AsStruct(), GetStructDeclId<Ids, InnermostType>()};
+			FBothStructId Id = {InnermostSchema.Get().AsStruct(), GetStructDeclId<Ids, InnermostType>()};
 			Init.RegisterInnerStruct(Id, Names);
 		}
 	}
diff --git a/Engine/Plugins/Experimental/PlainProps/Source/Public/PlainPropsTypes.h b/Engine/Plugins/Experimental/PlainProps/Source/Public/PlainPropsTypes.h
--- a/Engine/Plugins/Experimental/PlainProps/Source/Public/PlainPropsTypes.h
+++ b/Engine/Plugins/Experimental/PlainProps/Source/Public/PlainPropsTypes.h
@@ -134,7 +134,7 @@ struct FEnumId
 	friend bool operator==(FEnumId, FEnumId) = default;
 };
 
-// Runtime id for a struct FType
+// Runtime id for a struct FType, i.e. a class or struct
 struct FStructId
 {
 	uint32 Idx = ~0u;
@@ -142,6 +142,28 @@ struct FStructId
 	friend bool operator==(FStructId, FStructId) = default;
 };
 
+// Abstract FStructId used in declarations and stable schemas
+//
+// Might type-erase runtime details such as container allocator
+struct FDeclId : FStructId {};
+
+// Concrete FStructId used in bindings and in-memory schemas
+//
+// Uniquely ids a runtime class/struct. Usually same as FDeclId, 
+// but might type-erase runtime details.
+struct FBindId : FStructId {};
+
+// Static FBindId -> FDeclId cast once you've checked the bind id isnt type-erased
+inline FDeclId LowerCast(FBindId Id) { return static_cast<FDeclId>(FStructId(Id)); }
+inline FBindId UpCast(FDeclId Id) { return static_cast<FBindId>(FStructId(Id)); }
+
+// Interface for type-erasing / lowering FBindId to FDeclId
+struct IBindIds
+{
+	virtual ~IBindIds() = default;
+	virtual FDeclId Lower(FBindId BindId) const = 0;
+};
+
 // Either a runtime FStructId or FEnumId of a member
 struct FInnerId
 {
@@ -150,10 +172,12 @@ struct FInnerId
 	explicit FInnerId(FEnumId In) : Idx(In.Idx) {}
 	explicit FInnerId(FStructId In) : Idx(In.Idx) {}
 	
-	FEnumId		AsEnum() const		{ return {Idx}; }
-	FStructId	AsStruct() const	{ return {Idx}; }
+	FEnumId		AsEnum() const				{ return {Idx}; }
+	FStructId	AsStruct() const			{ return {Idx}; }
+	FDeclId		AsStructDeclId() const		{ return {Idx}; }
+	FBindId		AsStructBindId() const		{ return {Idx}; }
 
-	friend uint32 GetTypeHash(FInnerId In) { return GetTypeHash(In.Idx); };
+	friend uint32 GetTypeHash(FInnerId In)	{ return GetTypeHash(In.Idx); };
 	friend bool operator==(FInnerId, FInnerId) = default;
 	explicit FInnerId(uint32 InIdx) : Idx(InIdx) {}
 };
@@ -321,6 +345,8 @@ using FOptionalNameId = TOptionalId<FNameId>;
 using FOptionalMemberId = TOptionalId<FMemberId>;
 using FOptionalEnumId = TOptionalId<FEnumId>;
 using FOptionalStructId = TOptionalId<FStructId>;
+using FOptionalDeclId = TOptionalId<FDeclId>;
+using FOptionalBindId = TOptionalId<FBindId>;
 using FOptionalInnerId = TOptionalId<FInnerId>;
 using FOptionalSchemaId = TOptionalId<FSchemaId>;
 using FOptionalStructSchemaId = TOptionalId<FStructSchemaId>;
@@ -334,6 +360,9 @@ inline constexpr TOptionalId<IdType> ToOptional(IdType Id) { return Id; }
 
 inline FOptionalEnumId			ToOptionalEnum(FOptionalInnerId In)		{ return static_cast<const FOptionalEnumId&>(static_cast<const FOptionalId&>(In)); }
 inline FOptionalStructId		ToOptionalStruct(FOptionalInnerId In)	{ return static_cast<const FOptionalStructId&>(static_cast<const FOptionalId&>(In)); }
+inline FOptionalStructId		ToOptionalStruct(FOptionalDeclId In)	{ return static_cast<const FOptionalStructId&>(static_cast<const FOptionalId&>(In)); }
+inline FOptionalStructId		ToOptionalStruct(FOptionalBindId In)	{ return static_cast<const FOptionalStructId&>(static_cast<const FOptionalId&>(In)); }
+inline FOptionalDeclId			ToOptionalDeclId(FOptionalInnerId In)	{ return static_cast<const FOptionalDeclId&>(static_cast<const FOptionalId&>(In)); }
 inline FOptionalEnumSchemaId	ToOptionalEnum(FOptionalSchemaId In)	{ return static_cast<const FOptionalEnumSchemaId&>(static_cast<const FOptionalId&>(In)); }
 inline FOptionalStructSchemaId	ToOptionalStruct(FOptionalSchemaId In)	{ return static_cast<const FOptionalStructSchemaId&>(static_cast<const FOptionalId&>(In)); }
 
diff --git a/Engine/Plugins/Experimental/PlainProps/Source/Public/PlainPropsUeCoreBindings.h b/Engine/Plugins/Experimental/PlainProps/Source/Public/PlainPropsUeCoreBindings.h
--- a/Engine/Plugins/Experimental/PlainProps/Source/Public/PlainPropsUeCoreBindings.h
+++ b/Engine/Plugins/Experimental/PlainProps/Source/Public/PlainPropsUeCoreBindings.h
@@ -503,15 +503,15 @@ struct FTransformBinding : ICustomBinding
 {
 	using Type = FTransform;
 	enum class EMember : uint8 { Translate, Rotate, Scale };
-	const FMemberId MemberIds[3];
-	const FStructId VectorId;
-	const FStructId QuatId;
+	const FMemberId		MemberIds[3];
+	const FBindId		VectorId;
+	const FBindId		QuatId;
 
 	template<class Ids>
 	FTransformBinding(TCustomInit<Ids>)
 	: MemberIds{Ids::IndexMember("Translate"), Ids::IndexMember("Rotate"), Ids::IndexMember("Scale")}
-	, VectorId(GetStructDeclId<Ids, FVector>())
-	, QuatId(GetStructDeclId<Ids, FQuat>())
+	, VectorId(GetStructBindId<Ids, FVector>())
+	, QuatId(GetStructBindId<Ids, FQuat>())
 	{}
 
 	PLAINPROPS_API void	Save(FMemberBuilder& Dst, const FTransform& Src, const FTransform* Default, const FSaveContext& Context) const;
diff --git a/Engine/Plugins/Experimental/PlainProps/Source/Public/PlainPropsWrite.h b/Engine/Plugins/Experimental/PlainProps/Source/Public/PlainPropsWrite.h
--- a/Engine/Plugins/Experimental/PlainProps/Source/Public/PlainPropsWrite.h
+++ b/Engine/Plugins/Experimental/PlainProps/Source/Public/PlainPropsWrite.h
@@ -16,12 +16,12 @@ class FIdIndexerBase;
 class FNestedScopeIndexer;
 class FParametricTypeIndexer;
 struct FWriteIds;
-struct IStructBindIds;
+struct IBindIds;
 
 class FWriter
 {
 public:
-	PLAINPROPS_API FWriter(const FIdIndexerBase& AllIds, const IStructBindIds& BindIds, const FBuiltSchemas& InSchemas, ESchemaFormat Format);
+	PLAINPROPS_API FWriter(const FIdIndexerBase& AllIds, const IBindIds& BindIds, const FBuiltSchemas& InSchemas, ESchemaFormat Format);
 	PLAINPROPS_API ~FWriter();
 	
 	PLAINPROPS_API TConstArrayView<FNameId>			GetUsedNames() const;
diff --git a/Engine/Plugins/Experimental/PlainProps/Source/Tests/PlainPropsTests.cpp b/Engine/Plugins/Experimental/PlainProps/Source/Tests/PlainPropsTests.cpp
--- a/Engine/Plugins/Experimental/PlainProps/Source/Tests/PlainPropsTests.cpp
+++ b/Engine/Plugins/Experimental/PlainProps/Source/Tests/PlainPropsTests.cpp
@@ -123,27 +123,27 @@ void AppendString(FUtf8StringBuilderBase& Out, const FAnsiString& Str)
 	Out.Append(Str);
 }
 
-class FTestBatchBuilder : public TIdIndexer<FAnsiString>, public IStructBindIds
+class FTestBatchBuilder : public TIdIndexer<FAnsiString>, public IBindIds
 {
 public:
 	FTestBatchBuilder(FScratchAllocator& InScratch) : Declarations(FDebugIds(*this)), Scratch(InScratch) {}
 
 	FEnumId						DeclareEnum(FType Type, EEnumMode Mode, ELeafWidth Width, std::initializer_list<const char*> Names, std::initializer_list<uint64> Constants);
 	FEnumId						DeclareEnum(const char* Scope, const char* Name, EEnumMode Mode, ELeafWidth Width, std::initializer_list<const char*> Names, std::initializer_list<uint64> Constants);
-	FStructId					DeclareStruct(FType Type, std::initializer_list<const char*> MemberOrder, EMemberPresence Occupancy, FOptionalStructId Super = {});
-	FStructId					DeclareStruct(const char* Scope, const char* Name, std::initializer_list<const char*> MemberOrder, EMemberPresence Occupancy, FOptionalStructId Super = {});
+	FDeclId						DeclareStruct(FType Type, std::initializer_list<const char*> MemberOrder, EMemberPresence Occupancy, FOptionalDeclId Super = NoId);
+	FDeclId						DeclareStruct(const char* Scope, const char* Name, std::initializer_list<const char*> MemberOrder, EMemberPresence Occupancy, FOptionalDeclId Super = NoId);
 	
 	const FEnumDeclaration&		Get(FEnumId Id) { return Declarations.Get(Id); }
-	const FStructDeclaration&	Get(FStructId Id) { return Declarations.Get(Id); }
+	const FStructDeclaration&	Get(FDeclId Id) { return Declarations.Get(Id); }
 
-	void						AddObject(FStructId Schema, FMemberBuilder&& Members);
+	void						AddObject(FDeclId Id, FMemberBuilder&& Members);
 
 	TArray64<uint8>				Write();
 
 	FDebugIds					GetDebug() const { return FDebugIds(*this); }
 
 private:
-	TArray<TPair<FStructId, FBuiltStruct*>>		Objects;
+	TArray<TPair<FDeclId, FBuiltStruct*>>		Objects;
 	FDeclarations								Declarations;
 	FScratchAllocator&							Scratch;
 
@@ -172,21 +172,21 @@ private:
 
 	TArray<char> GetNameData() const;
 
-	virtual FStructId GetDeclId(FStructId BindId) const override final
+	virtual FDeclId Lower(FBindId Id) const override final
 	{
 		checkf(false, TEXT("All struct ids should be declared, nothing is bound with different names in this test suite"));
-		return BindId;
+		return LowerCast(Id);
 	}
 };
 
-FStructId FTestBatchBuilder::DeclareStruct(FType Type, std::initializer_list<const char*> MemberOrder, EMemberPresence Occupancy, FOptionalStructId Super)
+FDeclId FTestBatchBuilder::DeclareStruct(FType Type, std::initializer_list<const char*> MemberOrder, EMemberPresence Occupancy, FOptionalDeclId Super)
 {
-	FStructId Id = IndexStruct(Type);
+	FDeclId Id = IndexDeclId(Type);
 	Declarations.DeclareStruct(Id, Type, 0, NameMembers(MemberOrder), Occupancy, Super);
 	return Id;
 }
 
-FStructId FTestBatchBuilder::DeclareStruct(const char* Scope, const char* Name, std::initializer_list<const char*> MemberOrder, EMemberPresence Occupancy, FOptionalStructId Super)
+FDeclId FTestBatchBuilder::DeclareStruct(const char* Scope, const char* Name, std::initializer_list<const char*> MemberOrder, EMemberPresence Occupancy, FOptionalDeclId Super)
 {
 	return DeclareStruct(MakeType(Scope, Name), MemberOrder, Occupancy, Super);
 }
@@ -203,16 +203,16 @@ FEnumId FTestBatchBuilder::DeclareEnum(const char* Scope, const char* Name, EEnu
 	return DeclareEnum(MakeType(Scope, Name), Mode, Width, InNames, Constants);
 }
 
-void FTestBatchBuilder::AddObject(FStructId Schema, FMemberBuilder&& Members)
+void FTestBatchBuilder::AddObject(FDeclId Id, FMemberBuilder&& Members)
 {
-	Objects.Emplace(Schema, Members.BuildAndReset(Scratch, Declarations.Get(Schema), GetDebug()));
+	Objects.Emplace(Id, Members.BuildAndReset(Scratch, Declarations.Get(Id), GetDebug()));
 }
 
 TArray64<uint8> FTestBatchBuilder::Write()
 {
 	// Build partial schemas
 	FSchemasBuilder SchemaBuilders(Declarations, *this, *this, Scratch, ESchemaFormat::StableNames);
-	for (const TPair<FStructId, FBuiltStruct*>& Object : Objects)
+	for (const TPair<FDeclId, FBuiltStruct*>& Object : Objects)
 	{
 		SchemaBuilders.NoteStructAndMembers(Object.Key, *Object.Value);
 	}
@@ -243,7 +243,7 @@ TArray64<uint8> FTestBatchBuilder::Write()
 
 	// Write objects
 	WriteInt(Out, TestMagics[2]);
-	for (const TPair<FStructId, FBuiltStruct*>& Object : Objects)
+	for (const TPair<FDeclId, FBuiltStruct*>& Object : Objects)
 	{
 		WriteInt(/* out */ Tmp, TestMagics[3]);
 		WriteInt(/* out */ Tmp, Writer.GetWriteId(Object.Key).Get().Idx);
@@ -489,7 +489,7 @@ TEST_CASE_NAMED(FPlainPropsReadWriteTest, "System::Core::Serialization::PlainPro
 
 		TestSerialize([](FTestBatchBuilder& Batch, FScratchAllocator& Scratch)
 		{
-			FStructId SchemaId = Batch.DeclareStruct("Testing", "Bools", MemberNames, EMemberPresence::AllowSparse);
+			FDeclId Id = Batch.DeclareStruct("Testing", "Bools", MemberNames, EMemberPresence::AllowSparse);
 
 			FMemberBuilder B1T;
 			B1T.Add(Batch.NameMember("b3"),	true);
@@ -518,10 +518,10 @@ TEST_CASE_NAMED(FPlainPropsReadWriteTest, "System::Core::Serialization::PlainPro
 			B9T.Add(Batch.NameMember("b9"), true);
 			B9T.Add(Batch.NameMember("b10"),true);
 
-			Batch.AddObject(SchemaId, MoveTemp(B1T));
-			Batch.AddObject(SchemaId, MoveTemp(B1F));
-			Batch.AddObject(SchemaId, MoveTemp(B8M));
-			Batch.AddObject(SchemaId, MoveTemp(B9T));
+			Batch.AddObject(Id, MoveTemp(B1T));
+			Batch.AddObject(Id, MoveTemp(B1F));
+			Batch.AddObject(Id, MoveTemp(B8M));
+			Batch.AddObject(Id, MoveTemp(B9T));
 		}, 
 		[](TConstArrayView<FStructView> Objects, const FTestNameReader& Names)
 		{
@@ -582,7 +582,7 @@ TEST_CASE_NAMED(FPlainPropsReadWriteTest, "System::Core::Serialization::PlainPro
 
 		TestSerialize([](FTestBatchBuilder& Batch, FScratchAllocator& Scratch)
 		{
-			FStructId SchemaId = Batch.DeclareStruct("Test", "Numbers", MemberNames, EMemberPresence::AllowSparse);
+			FDeclId Id = Batch.DeclareStruct("Test", "Numbers", MemberNames, EMemberPresence::AllowSparse);
 
 			FMemberBuilder Misc, Mins, Maxs, Some;
 
@@ -621,10 +621,10 @@ TEST_CASE_NAMED(FPlainPropsReadWriteTest, "System::Core::Serialization::PlainPro
 
 			Some.Add(Batch.NameMember("S32"),	0);
 			
-			Batch.AddObject(SchemaId, MoveTemp(Misc));
-			Batch.AddObject(SchemaId, MoveTemp(Mins));
-			Batch.AddObject(SchemaId, MoveTemp(Maxs));
-			Batch.AddObject(SchemaId, MoveTemp(Some));
+			Batch.AddObject(Id, MoveTemp(Misc));
+			Batch.AddObject(Id, MoveTemp(Mins));
+			Batch.AddObject(Id, MoveTemp(Maxs));
+			Batch.AddObject(Id, MoveTemp(Some));
 		}, 
 		[](TConstArrayView<FStructView> Objects, const FTestNameReader& Names)
 		{
@@ -686,9 +686,9 @@ TEST_CASE_NAMED(FPlainPropsReadWriteTest, "System::Core::Serialization::PlainPro
 	{
 		TestSerialize([](FTestBatchBuilder& Batch, FScratchAllocator& Scratch)
 		{
-			FStructId Char8Id  = Batch.DeclareStruct("Test", "Char8",  {"A", "B", "C", "D", "E", "F"}, EMemberPresence::AllowSparse);
-			FStructId Char16Id = Batch.DeclareStruct("Test", "Char16", {"A", "B", "C", "D", "E", "F"}, EMemberPresence::AllowSparse);
-			FStructId Char32Id = Batch.DeclareStruct("Test", "Char32", {"A", "B", "C", "D", "E", "F"}, EMemberPresence::AllowSparse);
+			FDeclId Char8Id  = Batch.DeclareStruct("Test", "Char8",  {"A", "B", "C", "D", "E", "F"}, EMemberPresence::AllowSparse);
+			FDeclId Char16Id = Batch.DeclareStruct("Test", "Char16", {"A", "B", "C", "D", "E", "F"}, EMemberPresence::AllowSparse);
+			FDeclId Char32Id = Batch.DeclareStruct("Test", "Char32", {"A", "B", "C", "D", "E", "F"}, EMemberPresence::AllowSparse);
 
 			FMemberBuilder Char8, Char16, Char32;
 
@@ -749,8 +749,8 @@ TEST_CASE_NAMED(FPlainPropsReadWriteTest, "System::Core::Serialization::PlainPro
 
 		TestSerialize([](FTestBatchBuilder& Batch, FScratchAllocator& Scratch)
 		{
-			FStructId ExplicitId = Batch.DeclareStruct("Test", "ExplicitDense", ExplicitMembers, EMemberPresence::RequireAll);
-			FStructId ImplicitId = Batch.DeclareStruct("Test", "ImplicitDense", ImplicitMembers, EMemberPresence::AllowSparse);
+			FDeclId ExplicitId = Batch.DeclareStruct("Test", "ExplicitDense", ExplicitMembers, EMemberPresence::RequireAll);
+			FDeclId ImplicitId = Batch.DeclareStruct("Test", "ImplicitDense", ImplicitMembers, EMemberPresence::AllowSparse);
 
 			FMemberBuilder X;
 			X.Add(Batch.NameMember("A"), char8_t('a'));
@@ -808,11 +808,11 @@ TEST_CASE_NAMED(FPlainPropsReadWriteTest, "System::Core::Serialization::PlainPro
 
 		TestSerialize([](FTestBatchBuilder& Batch, FScratchAllocator& Scratch)
 		{
-			FStructId ObjectId = Batch.DeclareStruct("Test", "Object", ObjectMembers, EMemberPresence::AllowSparse);
-			FStructId StructId = Batch.DeclareStruct("Test", "Struct", StructMembers, EMemberPresence::AllowSparse);
-			FStructId NestedId = Batch.DeclareStruct("Test", "Nested", NestedMembers, EMemberPresence::AllowSparse);
-			FStructId UnusedId = Batch.DeclareStruct("Test", "Unused", UnusedMembers, EMemberPresence::AllowSparse);
-			FStructId EmptyId  = Batch.DeclareStruct("Test", "Empty",  EmptyMembers,  EMemberPresence::AllowSparse);
+			FDeclId ObjectId = Batch.DeclareStruct("Test", "Object", ObjectMembers, EMemberPresence::AllowSparse);
+			FDeclId StructId = Batch.DeclareStruct("Test", "Struct", StructMembers, EMemberPresence::AllowSparse);
+			FDeclId NestedId = Batch.DeclareStruct("Test", "Nested", NestedMembers, EMemberPresence::AllowSparse);
+			FDeclId UnusedId = Batch.DeclareStruct("Test", "Unused", UnusedMembers, EMemberPresence::AllowSparse);
+			FDeclId EmptyId  = Batch.DeclareStruct("Test", "Empty",  EmptyMembers,  EMemberPresence::AllowSparse);
 
 			FMemberBuilder Members;
 			Members.Add(Batch.NameMember("I1"), 100);
@@ -869,9 +869,9 @@ TEST_CASE_NAMED(FPlainPropsReadWriteTest, "System::Core::Serialization::PlainPro
 
 		TestSerialize([](FTestBatchBuilder& Batch, FScratchAllocator& Scratch)
 		{
-			// Test create holes in the original FNameId, FStructId and FEnumId index range
-			FStructId UnusedId = Batch.DeclareStruct("Test", "UnusedStruct", {"U1", "U2"}, EMemberPresence::AllowSparse);
-			FStructId ObjectId = Batch.DeclareStruct("Test", "Enums", MemberNames, EMemberPresence::AllowSparse);
+			// Test create holes in the original FNameId, FDeclId and FEnumId index range
+			FDeclId UnusedId = Batch.DeclareStruct("Test", "UnusedStruct", {"U1", "U2"}, EMemberPresence::AllowSparse);
+			FDeclId ObjectId = Batch.DeclareStruct("Test", "Enums", MemberNames, EMemberPresence::AllowSparse);
 
 			FEnumId U	= Batch.DeclareEnum("Test", "UnusedEnum1",	EEnumMode::Flag, ELeafWidth::B8,	{"U3"}, {1}); // Hole
 			FEnumId A	= Batch.DeclareEnum("Test", "FlatDense8",	EEnumMode::Flat, ELeafWidth::B8,	{"A", "B", "C"}, {0, 1, 2});
@@ -1002,7 +1002,7 @@ TEST_CASE_NAMED(FPlainPropsReadWriteTest, "System::Core::Serialization::PlainPro
 		{
 			static constexpr std::initializer_list<const char*> MemberNames = { "B0", "B1", "B8", "B9", "D0", "D3", "Hi", "E3", "E0" };
 		
-			FStructId ObjectId = Batch.DeclareStruct("Test", "Object", MemberNames, EMemberPresence::AllowSparse);
+			FDeclId ObjectId = Batch.DeclareStruct("Test", "Object", MemberNames, EMemberPresence::AllowSparse);
 			FEnumId Enum	= Batch.DeclareEnum("Test", "ABCD", EEnumMode::Flat, ELeafWidth::B16, {"A", "B", "C", "D"}, {0, 1, 2, 3});
 			FEnumId UnusedEnum1	= Batch.DeclareEnum("Test", "Unused1", EEnumMode::Flat, ELeafWidth::B8, {"X"}, {0});
 			FEnumId UnusedEnum2	= Batch.DeclareEnum("Test", "Unused2", EEnumMode::Flat, ELeafWidth::B8, {"Y"}, {0});
@@ -1051,8 +1051,8 @@ TEST_CASE_NAMED(FPlainPropsReadWriteTest, "System::Core::Serialization::PlainPro
 	{
 		TestSerialize([](FTestBatchBuilder& Batch, FScratchAllocator& Scratch)
 		{
-			FStructId Utf8Id  = Batch.DeclareStruct("Test", "Utf8",  {"Null", "Empty", "Escape", "Latin1", "CJK", "Symbols"}, EMemberPresence::AllowSparse);
-			FStructId Utf16Id = Batch.DeclareStruct("Test", "Utf16", {"Null", "Empty", "Escape", "Latin1", "CJK", "Symbols"}, EMemberPresence::AllowSparse);
+			FDeclId Utf8Id  = Batch.DeclareStruct("Test", "Utf8",  {"Null", "Empty", "Escape", "Latin1", "CJK", "Symbols"}, EMemberPresence::AllowSparse);
+			FDeclId Utf16Id = Batch.DeclareStruct("Test", "Utf16", {"Null", "Empty", "Escape", "Latin1", "CJK", "Symbols"}, EMemberPresence::AllowSparse);
 
 			FMemberBuilder Utf8, Utf16;
 
@@ -1099,8 +1099,8 @@ TEST_CASE_NAMED(FPlainPropsReadWriteTest, "System::Core::Serialization::PlainPro
 	{
 		TestSerialize([](FTestBatchBuilder& Batch, FScratchAllocator& Scratch)
 		{
-			FStructId ObjectId = Batch.DeclareStruct("Test", "Object", {"Structs"}, EMemberPresence::AllowSparse);
-			FStructId StructId = Batch.DeclareStruct("Test", "Struct", {"I", "F"}, EMemberPresence::AllowSparse);
+			FDeclId ObjectId = Batch.DeclareStruct("Test", "Object", {"Structs"}, EMemberPresence::AllowSparse);
+			FDeclId StructId = Batch.DeclareStruct("Test", "Struct", {"I", "F"}, EMemberPresence::AllowSparse);
 			 
 			FStructRangeBuilder Structs(4);
 			Structs[0].Add(Batch.NameMember("I"), 0);
@@ -1135,9 +1135,9 @@ TEST_CASE_NAMED(FPlainPropsReadWriteTest, "System::Core::Serialization::PlainPro
 
 		TestSerialize([](FTestBatchBuilder& Batch, FScratchAllocator& Scratch)
 		{
-			FStructId Object = Batch.DeclareStruct("Test", "Object", {"IntRs", "EmptyRs", "FloatRs", "EnumRs", "UnusedEnumRs", "UnicodeRs", "StructRs", "StructRRs"}, EMemberPresence::AllowSparse);
-			FStructId XY = Batch.DeclareStruct("Test", "XY", {"X", "Y"}, EMemberPresence::RequireAll);
-			FStructId ZW = Batch.DeclareStruct("Test", "ZW", {"Z", "W"}, EMemberPresence::AllowSparse);
+			FDeclId Object = Batch.DeclareStruct("Test", "Object", {"IntRs", "EmptyRs", "FloatRs", "EnumRs", "UnusedEnumRs", "UnicodeRs", "StructRs", "StructRRs"}, EMemberPresence::AllowSparse);
+			FDeclId XY = Batch.DeclareStruct("Test", "XY", {"X", "Y"}, EMemberPresence::RequireAll);
+			FDeclId ZW = Batch.DeclareStruct("Test", "ZW", {"Z", "W"}, EMemberPresence::AllowSparse);
 			FEnumId Enum = Batch.DeclareEnum("Test", "AB", EEnumMode::Flag, ELeafWidth::B8, {"A", "B"}, {1, 4});
 			FEnumId UnusedEnum = Batch.DeclareEnum("Test", "Unused", EEnumMode::Flat, ELeafWidth::B8, {"X"}, {0});
 	
@@ -1258,8 +1258,8 @@ TEST_CASE_NAMED(FPlainPropsReadWriteTest, "System::Core::Serialization::PlainPro
 	{
 		TestSerialize([](FTestBatchBuilder& Batch, FScratchAllocator& Scratch)
 		{
-			FStructId Object = Batch.DeclareStruct("Test", "Object", {"Bools", "Structs", "BF", "BT" }, EMemberPresence::AllowSparse);
-			FStructId Struct = Batch.DeclareStruct("Test", "Struct", {"MaybeB", "Bs", "MaybeBs", "B"}, EMemberPresence::AllowSparse);
+			FDeclId Object = Batch.DeclareStruct("Test", "Object", {"Bools", "Structs", "BF", "BT" }, EMemberPresence::AllowSparse);
+			FDeclId Struct = Batch.DeclareStruct("Test", "Struct", {"MaybeB", "Bs", "MaybeBs", "B"}, EMemberPresence::AllowSparse);
 						
 			const bool True = true;
 			const bool False = false;
@@ -1316,14 +1316,14 @@ TEST_CASE_NAMED(FPlainPropsReadWriteTest, "System::Core::Serialization::PlainPro
 	{
 		TestSerialize([](FTestBatchBuilder& Batch, FScratchAllocator& Scratch)
 		{
-			FStructId Unused1 = Batch.DeclareStruct("Test", "Unused1", {"X"}, EMemberPresence::AllowSparse);
-			FStructId SA = Batch.DeclareStruct("Test", "SA", {"X"}, EMemberPresence::AllowSparse);
-			FStructId Unused2 = Batch.DeclareStruct("Test", "Unused2", {"X"}, EMemberPresence::AllowSparse);
-			FStructId SB = Batch.DeclareStruct("Test", "SB", {"X"}, EMemberPresence::AllowSparse);
-			FStructId Object = Batch.DeclareStruct("Test", "Object", {"Same", "Some", "None", "Diff"}, EMemberPresence::AllowSparse);
-			FStructId Unused3 = Batch.DeclareStruct("Test", "Unused3", {"X"}, EMemberPresence::AllowSparse);
+			FDeclId Unused1 = Batch.DeclareStruct("Test", "Unused1", {"X"}, EMemberPresence::AllowSparse);
+			FDeclId SA = Batch.DeclareStruct("Test", "SA", {"X"}, EMemberPresence::AllowSparse);
+			FDeclId Unused2 = Batch.DeclareStruct("Test", "Unused2", {"X"}, EMemberPresence::AllowSparse);
+			FDeclId SB = Batch.DeclareStruct("Test", "SB", {"X"}, EMemberPresence::AllowSparse);
+			FDeclId Object = Batch.DeclareStruct("Test", "Object", {"Same", "Some", "None", "Diff"}, EMemberPresence::AllowSparse);
+			FDeclId Unused3 = Batch.DeclareStruct("Test", "Unused3", {"X"}, EMemberPresence::AllowSparse);
 		
-			auto BuildStruct = [&](FStructId Struct, auto X)
+			auto BuildStruct = [&](FDeclId Struct, auto X)
 				{
 					FMemberBuilder Members;
 					Members.Add(Batch.NameMember("X"), X);
@@ -1365,12 +1365,12 @@ TEST_CASE_NAMED(FPlainPropsReadWriteTest, "System::Core::Serialization::PlainPro
 	{
 		TestSerialize([](FTestBatchBuilder& Batch, FScratchAllocator& Scratch)
 		{
-			FStructId SA = Batch.DeclareStruct("Test", "SA", {"X"}, EMemberPresence::AllowSparse);
-			FStructId Unused = Batch.DeclareStruct("Test", "Unused2", {"X"}, EMemberPresence::AllowSparse);
-			FStructId SB = Batch.DeclareStruct("Test", "SB", {"X"}, EMemberPresence::AllowSparse);
-			FStructId Object = Batch.DeclareStruct("Test", "Object", {"Same", "Some", "None", "Diff", "SameEmpty", "DiffEmpty", "DiffNested"}, EMemberPresence::AllowSparse);
+			FDeclId SA = Batch.DeclareStruct("Test", "SA", {"X"}, EMemberPresence::AllowSparse);
+			FDeclId Unused = Batch.DeclareStruct("Test", "Unused2", {"X"}, EMemberPresence::AllowSparse);
+			FDeclId SB = Batch.DeclareStruct("Test", "SB", {"X"}, EMemberPresence::AllowSparse);
+			FDeclId Object = Batch.DeclareStruct("Test", "Object", {"Same", "Some", "None", "Diff", "SameEmpty", "DiffEmpty", "DiffNested"}, EMemberPresence::AllowSparse);
 		
-			auto BuildStructRange = [&](FStructId Struct, auto X)
+			auto BuildStructRange = [&](FDeclId Struct, auto X)
 				{
 					FStructRangeBuilder Members(1);
 					Members[0].Add(Batch.NameMember("X"), X);
@@ -1428,10 +1428,10 @@ TEST_CASE_NAMED(FPlainPropsReadWriteTest, "System::Core::Serialization::PlainPro
 	{
 		TestSerialize([](FTestBatchBuilder& Batch, FScratchAllocator& Scratch)
 		{
-			FStructId Unused = Batch.DeclareStruct("Test", "X", {"X"}, EMemberPresence::AllowSparse);
-			FStructId Low = Batch.DeclareStruct("Test", "Low", {"LInt"}, EMemberPresence::AllowSparse);
-			FStructId Mid = Batch.DeclareStruct("Test", "Mid", {"MInt", "MLow"}, EMemberPresence::AllowSparse, ToOptional(Low));
-			FStructId Top = Batch.DeclareStruct("Test", "Top", {"TInt", "TLow", "TMids"}, EMemberPresence::AllowSparse, ToOptional(Mid));
+			FDeclId Unused = Batch.DeclareStruct("Test", "X", {"X"}, EMemberPresence::AllowSparse);
+			FDeclId Low = Batch.DeclareStruct("Test", "Low", {"LInt"}, EMemberPresence::AllowSparse);
+			FDeclId Mid = Batch.DeclareStruct("Test", "Mid", {"MInt", "MLow"}, EMemberPresence::AllowSparse, ToOptional(Low));
+			FDeclId Top = Batch.DeclareStruct("Test", "Top", {"TInt", "TLow", "TMids"}, EMemberPresence::AllowSparse, ToOptional(Mid));
 			
 			FMemberBuilder Members;
 			Members.Add(Batch.NameMember("LInt"), 123);
@@ -1509,14 +1509,14 @@ TEST_CASE_NAMED(FPlainPropsReadWriteTest, "System::Core::Serialization::PlainPro
 	{
 		TestSerialize([](FTestBatchBuilder& Batch, FScratchAllocator& Scratch)
 		{																										// Usage by A B C 
-			FStructId B0 = Batch.DeclareStruct("Test", "B0", {"0"}, EMemberPresence::AllowSparse);				 // - - -
-			FStructId B1 = Batch.DeclareStruct("Test", "B1", {"1"}, EMemberPresence::AllowSparse, ToOptional(B0)); // 1 1 1
-			FStructId B2 = Batch.DeclareStruct("Test", "B2", {"2"}, EMemberPresence::AllowSparse, ToOptional(B1)); // - 1 1
-			FStructId B3 = Batch.DeclareStruct("Test", "B3", {"3"}, EMemberPresence::AllowSparse, ToOptional(B2)); // - - 0
-			FStructId B4 = Batch.DeclareStruct("Test", "B4", {"4"}, EMemberPresence::AllowSparse, ToOptional(B3)); // 1 1 1
-			FStructId B5 = Batch.DeclareStruct("Test", "B5", {"5"}, EMemberPresence::AllowSparse, ToOptional(B4)); // 1 1 0
-			FStructId B6 = Batch.DeclareStruct("Test", "B6", {   }, EMemberPresence::AllowSparse, ToOptional(B5)); // 0 - 0
-			FStructId C5 = Batch.DeclareStruct("Test", "C5", {"5"}, EMemberPresence::AllowSparse, ToOptional(B4)); // - - -
+			FDeclId B0 = Batch.DeclareStruct("Test", "B0", {"0"}, EMemberPresence::AllowSparse);				 // - - -
+			FDeclId B1 = Batch.DeclareStruct("Test", "B1", {"1"}, EMemberPresence::AllowSparse, ToOptional(B0)); // 1 1 1
+			FDeclId B2 = Batch.DeclareStruct("Test", "B2", {"2"}, EMemberPresence::AllowSparse, ToOptional(B1)); // - 1 1
+			FDeclId B3 = Batch.DeclareStruct("Test", "B3", {"3"}, EMemberPresence::AllowSparse, ToOptional(B2)); // - - 0
+			FDeclId B4 = Batch.DeclareStruct("Test", "B4", {"4"}, EMemberPresence::AllowSparse, ToOptional(B3)); // 1 1 1
+			FDeclId B5 = Batch.DeclareStruct("Test", "B5", {"5"}, EMemberPresence::AllowSparse, ToOptional(B4)); // 1 1 0
+			FDeclId B6 = Batch.DeclareStruct("Test", "B6", {   }, EMemberPresence::AllowSparse, ToOptional(B5)); // 0 - 0
+			FDeclId C5 = Batch.DeclareStruct("Test", "C5", {"5"}, EMemberPresence::AllowSparse, ToOptional(B4)); // - - -
 			
 			FMemberBuilder A;
 			A.Add(Batch.NameMember("1"), 1);
@@ -1626,11 +1626,11 @@ TEST_CASE_NAMED(FPlainPropsReadWriteTest, "System::Core::Serialization::PlainPro
 			FType S4T = Batch.MakeParametricType({DoubleNested,	Batch.MakeTypename("S4")}, {S2T, E2T});
 			FType S5T = Batch.MakeParametricType({NestedUnused3,	Batch.MakeTypename("S5")}, {E3T, E1T, S2T});
 			
-			FStructId S1D = Batch.DeclareStruct(S1T, {"M1"}, EMemberPresence::AllowSparse);
-			FStructId S2D = Batch.DeclareStruct(S2T, {"M2"}, EMemberPresence::AllowSparse);
-			FStructId S3D = Batch.DeclareStruct(S3T, {"M3"}, EMemberPresence::AllowSparse);
-			FStructId S4D = Batch.DeclareStruct(S4T, {"M4"}, EMemberPresence::AllowSparse);
-			FStructId S5D = Batch.DeclareStruct(S5T, {"M5"}, EMemberPresence::AllowSparse);
+			FDeclId S1D = Batch.DeclareStruct(S1T, {"M1"}, EMemberPresence::AllowSparse);
+			FDeclId S2D = Batch.DeclareStruct(S2T, {"M2"}, EMemberPresence::AllowSparse);
+			FDeclId S3D = Batch.DeclareStruct(S3T, {"M3"}, EMemberPresence::AllowSparse);
+			FDeclId S4D = Batch.DeclareStruct(S4T, {"M4"}, EMemberPresence::AllowSparse);
+			FDeclId S5D = Batch.DeclareStruct(S5T, {"M5"}, EMemberPresence::AllowSparse);
 
 			FMemberBuilder Members;
 			Members.Add(Batch.NameMember("M4"), 1);
diff --git a/Engine/Plugins/Experimental/PlainProps/Source/Tests/PlainPropsUeCoreTests.cpp b/Engine/Plugins/Experimental/PlainProps/Source/Tests/PlainPropsUeCoreTests.cpp
--- a/Engine/Plugins/Experimental/PlainProps/Source/Tests/PlainPropsUeCoreTests.cpp
+++ b/Engine/Plugins/Experimental/PlainProps/Source/Tests/PlainPropsUeCoreTests.cpp
@@ -93,7 +93,7 @@ using TScopedDefaultStructBinding = TScopedStructBinding<T, FDefaultRuntime>;
 // Maybe replace with macro, e.g. PP_DECLARE_CUSTOM_DENSE_STRUCT(FRuntimeIds, FName, void, Idx)
 struct FNameDeclaration 
 {
-	FStructId			Id;
+	FDeclId				Id;
 	FMemberId			Idx;
 
 	FNameDeclaration(FType Type = IndexStructName<FRuntimeIds, ETypename::Decl, TTypename<FName>>())
@@ -111,7 +111,7 @@ struct FNameDeclaration
 
 struct FTestCustomBinding : ICustomBinding
 {
-//	virtual FDualStructId GetId() const = 0;
+//	virtual FBothStructId GetId() const = 0;
 };
 
 struct FNameBinding : FTestCustomBinding
@@ -135,7 +135,7 @@ struct FNameBinding : FTestCustomBinding
 		return A.IsEqual(B, ENameCase::CaseSensitive);
 	}
 
-//	virtual FDualStructId GetId() const override { return { Declaration.Id, Declaration.Id }; }
+//	virtual FBothStructId GetId() const override { return { Declaration.Id, Declaration.Id }; }
 
 	FNameDeclaration	Declaration;
 	TSet<FName>			Names;
@@ -166,20 +166,20 @@ private:
 FBatchSaver::FBatchSaver(const FCustomBindings& Underlay)
 : Customs(Underlay)
 {
-	Customs.BindStruct(SavedNames.Declaration.Id, SavedNames, GTypes.Get(SavedNames.Declaration.Id), {});
+	Customs.BindStruct(UpCast(SavedNames.Declaration.Id), SavedNames, GTypes.Get(SavedNames.Declaration.Id), {});
 }
 
 template<class T>
 void FBatchSaver::Save(T&& Object) 
 {
-	FDualStructId Id = IndexStructDualId<FRuntimeIds, TTypename<std::remove_reference_t<T>>>();
+	FBothStructId Id = IndexStructBothId<FRuntimeIds, TTypename<std::remove_reference_t<T>>>();
 	SavedObjects.Emplace(Id.DeclId, SaveStruct(&Object, Id.BindId, {GTypes, GSchemas, Customs, Scratch}));
 }
 
 template<class T>
 bool FBatchSaver::SaveDelta(const T& Object, const T& Default) 
 {
-	FDualStructId Id = IndexStructDualId<FRuntimeIds, TTypename<std::remove_reference_t<T>>>();
+	FBothStructId Id = IndexStructBothId<FRuntimeIds, TTypename<std::remove_reference_t<T>>>();
 	if (const FBuiltStruct* Delta = SaveStructDeltaIfDiff(&Object, &Default, Id.BindId, {GTypes, GSchemas, Customs, Scratch}))
 	{
 		SavedObjects.Emplace(Id.DeclId, MoveTemp(Delta));
@@ -411,7 +411,7 @@ public:
 
 		// Read names and bind custom loader
 		NameBinding.Names.Append(GrabNumAndArray<FName>(It));
-		Customs.BindStruct(NameBinding.Declaration.Id, NameBinding, GTypes.Get(NameBinding.Declaration.Id), {});
+		Customs.BindStruct(UpCast(NameBinding.Declaration.Id), NameBinding, GTypes.Get(NameBinding.Declaration.Id), {});
 		CHECK(It.Grab<uint32>() == Magics[5]);
 
 		// Finally create load plans
diff --git a/Engine/Plugins/Experimental/PlainPropsUObject/Source/Private/PlainPropsUObjectBindings.cpp b/Engine/Plugins/Experimental/PlainPropsUObject/Source/Private/PlainPropsUObjectBindings.cpp
--- a/Engine/Plugins/Experimental/PlainPropsUObject/Source/Private/PlainPropsUObjectBindings.cpp
+++ b/Engine/Plugins/Experimental/PlainPropsUObject/Source/Private/PlainPropsUObjectBindings.cpp
@@ -142,7 +142,7 @@ static void ReserveZeroes(/* in-out */ FMutableMemoryView& Zeroes, SIZE_T Size,
 
 FDefaultStructs::~FDefaultStructs()
 {
-	for (TPair<FStructId, FDefaultInstance> Instance : Instances)
+	for (TPair<FBindId, FDefaultInstance> Instance : Instances)
 	{
 		DeleteInstance(Instance.Value);
 	}
@@ -165,7 +165,7 @@ void FDefaultStructs::ReserveFlags(uint32 Idx)
 	}
 }
 
-void FDefaultStructs::Bind(FStructId Id, const UScriptStruct* Struct)
+void FDefaultStructs::Bind(FBindId Id, const UScriptStruct* Struct)
 {
 	const EStructFlags Flags = Struct->StructFlags;
 	const SIZE_T Size = Struct->GetStructureSize();
@@ -212,7 +212,7 @@ void FDefaultStructs::Bind(FStructId Id, const UScriptStruct* Struct)
 	}
 }
 
-void FDefaultStructs::BindZeroes(FStructId Id, SIZE_T Size, uint32 Alignment)
+void FDefaultStructs::BindZeroes(FBindId Id, SIZE_T Size, uint32 Alignment)
 {
 	ReserveFlags(Id.Idx);
 #if DO_CHECK
@@ -221,7 +221,7 @@ void FDefaultStructs::BindZeroes(FStructId Id, SIZE_T Size, uint32 Alignment)
 	ReserveZeroes(Zeroes, Size, Alignment);
 }
 
-void FDefaultStructs::BindStatic(FStructId Id, const void* Struct)
+void FDefaultStructs::BindStatic(FBindId Id, const void* Struct)
 {
 	ReserveFlags(Id.Idx);
 #if DO_CHECK
@@ -234,7 +234,7 @@ void FDefaultStructs::BindStatic(FStructId Id, const void* Struct)
 	Instances.Add(Id, MakeStaticInstance(Struct));
 }
 
-void FDefaultStructs::Drop(FStructId Id)
+void FDefaultStructs::Drop(FBindId Id)
 {
 #if DO_CHECK
 	checkf(!Flip(Bound[Id.Idx]), TEXT("'%s' isn't bound"), *GUE.Debug.Print(Id));
@@ -246,7 +246,7 @@ void FDefaultStructs::Drop(FStructId Id)
 	}
 }
 
-const void* FDefaultStructs::Get(FStructId Id)
+const void* FDefaultStructs::Get(FBindId Id)
 {
 #if DO_CHECK
 	checkf(Bound[Id.Idx], TEXT("'%s' lack default"), *GUE.Debug.Print(Id));
@@ -359,7 +359,7 @@ static constexpr uint64 CommonStructMask =	CASTCLASS_FNameProperty | CASTCLASS_F
 											CASTCLASS_FDelegateProperty | CASTCLASS_FMulticastInlineDelegateProperty;
 static constexpr uint64 MiscMask =			CASTCLASS_FMulticastSparseDelegateProperty | CASTCLASS_FInterfaceProperty;
 
-static FStructId FlagsToCommonBindId(uint64 MaskedCastFlags)
+static FBindId FlagsToCommonBindId(uint64 MaskedCastFlags)
 {
 	switch (MaskedCastFlags)
 	{
@@ -456,14 +456,14 @@ static FType IndexType(const UField* Field)
 	return { GUE.Names.NestReversedScopes(ReversedOuters), Name };
 }
 
-static FOptionalStructId IndexSuper(const UStruct* Struct)
+static FOptionalDeclId IndexSuper(const UStruct* Struct)
 {
 	if (const UStruct* Super = Struct->GetInheritanceSuper())
 	{
 		if (ShouldBind(Super))
 		{
 			const UStruct* NonEmptySuper = SkipEmptyBases(Super);
-			return GUE.Names.IndexStruct(IndexType(NonEmptySuper));
+			return GUE.Names.IndexDeclId(IndexType(NonEmptySuper));
 		}
 	}
 	
@@ -506,9 +506,9 @@ static void DeclareSuperMembers(FMemberArray& Out, const UStruct* Struct)
 	}
 }
 
-static void DeclareStruct(const UStruct* Struct, FType Type, FStructId Id)
+static void DeclareStruct(const UStruct* Struct, FType Type, FDeclId Id)
 {
-	FOptionalStructId  Super = IndexSuper(Struct);
+	FOptionalDeclId Super = IndexSuper(Struct);
 	EMemberPresence Occupancy = GetOccupancy(Struct);
 	
 	FMemberArray Members;
@@ -523,10 +523,10 @@ static void DeclareStruct(const UStruct* Struct, FType Type, FStructId Id)
 	GUE.Types.DeclareStruct(Id, Type, 0, Members, Occupancy, Super);
 }
 
-static FStructId DeclareStruct(const UStruct* Struct)
+static FDeclId DeclareStruct(const UStruct* Struct)
 {
 	FType Type = IndexType(Struct);
-	FStructId Id = GUE.Names.IndexStruct(Type);
+	FDeclId Id = GUE.Names.IndexDeclId(Type);
 	DeclareStruct(Struct, Type, Id);
 	return Id;
 }
@@ -535,7 +535,7 @@ static FStructId DeclareStruct(const UStruct* Struct)
 
 static FTypedRange SaveNames(TConstArrayView<FName> Names, const FSaveContext& Ctx)
 {
-	FStructId Id = GUE.Structs.Name;
+	FBindId Id = GUE.Structs.Name;
 	FStructRangeSaver Out(Ctx.Scratch, Names.Num());
 	for (FName Name : Names)
 	{
@@ -622,7 +622,7 @@ static void LoadUnicastDelegate(FUnicastScriptDelegate& Dst, FStructLoadView Src
 
 static FTypedRange SaveInvocations(const FMulticastInvocationList& In, const FSaveContext& Ctx)
 {
-	FStructId ItemId = GUE.Structs.Delegate;
+	FBindId ItemId = GUE.Structs.Delegate;
 	FMemberSchema Schema = MakeStructRangeSchema(DefaultRangeMax, ItemId);
 	if (int32 NumTotal = In.Num())
 	{
@@ -635,7 +635,7 @@ static FTypedRange SaveInvocations(const FMulticastInvocationList& In, const FSa
 	
 		if (int32 NumKept = Keep.CountSetBits())
 		{
-			const FStructDeclaration& ItemDecl = GUE.Types.Get(ItemId);
+			const FStructDeclaration& ItemDecl = GUE.Types.Get(LowerCast(ItemId));
 			const FUnicastScriptDelegate* Src = In.GetData();
 			FStructRangeSaver Dst(Ctx.Scratch, static_cast<uint64>(NumKept));
 			FMemberBuilder Tmp;
@@ -812,7 +812,7 @@ struct FMulticastSparseDelegateBinding final : ICustomBinding
 	}
 };
 
-static FStructId BindSparseDelegate(FStructId Owner, FMulticastSparseDelegateProperty* Property)
+static FBindId BindSparseDelegate(FBindId Owner, FMulticastSparseDelegateProperty* Property)
 {
 	// Todo: Ownership / memory leak
 	ICustomBinding* Leak = new FMulticastSparseDelegateBinding(CastChecked<USparseDelegateFunction>(Property->SignatureFunction));
@@ -821,11 +821,11 @@ static FStructId BindSparseDelegate(FStructId Owner, FMulticastSparseDelegatePro
 	FType OwnerParam = GUE.Names.Resolve(Owner);
 	FType PropertyParam = { GUE.Scopes.CoreUObject, FTypenameId(GUE.Names.NameType(Property->GetFName())) };
 	FType UniqueBindName = GUE.Names.MakeParametricType(MulticastSparseDelegate, {OwnerParam, PropertyParam});
-	FStructId BindId = GUE.Names.IndexStruct(UniqueBindName);
+	FBindId Id = GUE.Names.IndexBindId(UniqueBindName);
 
-	GUE.Customs.BindStruct(BindId, *Leak, GUE.Types.Get(GUE.Structs.MulticastDelegate));
+	GUE.Customs.BindStruct(Id, *Leak, GUE.Types.Get(GUE.Structs.MulticastDelegate));
 
-	return BindId;
+	return Id;
 }
 
 ////////////////////////////////////////////////////////////////////////////////////////////////
@@ -928,7 +928,7 @@ class FInterfaceBindings
 	const FType ScriptInterface;
 	const FStructDeclaration& Declaration;
 
-	TMap<FType, FStructId> BoundClasses;
+	TMap<FType, FBindId> BoundClasses;
 
 public:
 	FInterfaceBindings()
@@ -936,16 +936,16 @@ public:
 	, Declaration(GUE.Types.DeclareStruct(GUE.Structs.ScriptInterface, ScriptInterface, 0, {GUE.Members.Object}, EMemberPresence::RequireAll))
 	{}
 
-	FStructId Bind(FInterfaceProperty* Property)
+	FBindId Bind(FInterfaceProperty* Property)
 	{
 		FType Class = IndexType(Property->InterfaceClass);
-		if (const FStructId* BindId = BoundClasses.Find(Class))
+		if (const FBindId* BindId = BoundClasses.Find(Class))
 		{
 			return *BindId;
 		}
 		
 		FType UniqueBindName = GUE.Names.MakeParametricType(ScriptInterface, {Class});
-		FStructId BindId = GUE.Names.IndexStruct(UniqueBindName);
+		FBindId BindId = GUE.Names.IndexBindId(UniqueBindName);
 		BoundClasses.Emplace(Class, BindId);
 
 		// Todo: Ownership / memory leak
@@ -956,7 +956,7 @@ public:
 	}
 };
 
-static FStructId BindInterface(FInterfaceProperty* Property)
+static FBindId BindInterface(FInterfaceProperty* Property)
 {
 	static FInterfaceBindings Bindings;
 	return Bindings.Bind(Property);
@@ -1640,7 +1640,7 @@ struct FStructRangeSerializer
 	using RangeSaver = FStructRangeSaver;
 	
 	FMemberType				InnerType;
-	FStructId				SaveId;
+	FBindId					SaveId;
 
 	explicit FStructRangeSerializer(FMemberBinding Item)
 	: InnerType(Item.InnermostType.AsStruct())
@@ -2226,17 +2226,17 @@ struct TSetPropertyBinding : IItemRangeBinding, ICustomBinding
 
 class FSetBindings
 {
-	TMap<FParameterBinding, FStructId> Bindings;
+	TMap<FParameterBinding, FBindId> Bindings;
 
 	template<EPropertyKind ItemKind>
-	void BindNew(FSetProperty* Property, FMemberBinding Elem, FStructId BindId, const FStructDeclaration& Declaration) 
+	void BindNew(FSetProperty* Property, FMemberBinding Elem, FBindId BindId, const FStructDeclaration& Declaration) 
 	{
 		// Todo: Ownership / memory leak
 		TSetPropertyBinding<ItemKind>* Leak = new TSetPropertyBinding<ItemKind>(Property, Elem);
 		GUE.Customs.BindStruct(BindId, *Leak, Declaration, {});
 	}
 
-	void DeltaBindNew(FSetProperty* Property, FMemberBinding Elem, FStructId BindId, const FStructDeclaration& Declaration)
+	void DeltaBindNew(FSetProperty* Property, FMemberBinding Elem, FBindId BindId, const FStructDeclaration& Declaration)
 	{
 		switch (GetPropertyKind(Elem))
 		{
@@ -2254,22 +2254,22 @@ class FSetBindings
 	}
 public:
 
-	FStructId Bind(FSetProperty* Property, FMemberBinding Elem)
+	FBindId Bind(FSetProperty* Property, FMemberBinding Elem)
 	{
 		check(Elem.Offset == 0);
-		if (const FStructId* BindId = Bindings.Find(FParameterBinding(Elem)))
+		if (const FBindId* BindId = Bindings.Find(FParameterBinding(Elem)))
 		{
 			return *BindId;
 		}
 
 		// Index custom delta binding struct name
-		FDualType Param = Elem.IndexParameterName(GUE.Names);
+		FBothType Param = Elem.IndexParameterName(GUE.Names);
 		FType BindType = FType{ GUE.Scopes.Core, FTypenameId(GUE.Names.MakeParametricTypeId(GUE.Typenames.Set, MakeArrayView(&Param.BindType, 1))) };
 		FType DeclType = Param.IsLowered()
 						 ? FType{ GUE.Scopes.Core, FTypenameId(GUE.Names.MakeParametricTypeId(GUE.Typenames.Set, MakeArrayView(&Param.DeclType, 1))) }
 						 : BindType;
-		FStructId BindId = GUE.Names.IndexStruct(BindType);
-		FStructId DeclId = Param.IsLowered() ? GUE.Names.IndexStruct(DeclType) : BindId;
+		FBindId BindId = GUE.Names.IndexBindId(BindType);
+		FDeclId DeclId = Param.IsLowered() ? GUE.Names.IndexDeclId(DeclType) : LowerCast(BindId);
 		FMemberId Members[] = { GUE.Members.Assign, GUE.Members.Remove, GUE.Members.Insert };
 
 		// Todo: Ownership / memory leak
@@ -2589,11 +2589,11 @@ struct TMapPropertyCustomBinding : TMapPropertyItemBinding<FScriptMap, KeyKind,
 
 class FMapBindings
 {
-	TMap<FStructId, FStructId>	NormalBindings;
-	TMap<FStructId, FRangeBinding>	FrozenBindings;
+	TMap<FBindId, FBindId>			NormalBindings;
+	TMap<FBindId, FRangeBinding>	FrozenBindings;
 
 	template<EPropertyKind KeyKind, EPropertyKind ValueKind>
-	static IItemRangeBinding* New3(FMapProperty* Property, FMapMemberBindings Members, FStructId* OutCustomId)
+	static IItemRangeBinding* New3(FMapProperty* Property, FMapMemberBindings Members, FBindId* OutCustomId)
 	{
 		if (OutCustomId == nullptr) // Freezable maps aren't delta serialized
 		{
@@ -2604,20 +2604,20 @@ class FMapBindings
 		FParametricTypeId PairTypename = GUE.Names.Resolve(Members.Pair.InnermostSchema.Get().AsStruct()).Name.AsParametric();
 		TConstArrayView<FType> Params = GUE.Names.Resolve(PairTypename).GetParameters();
 		FType Type = { GUE.Scopes.Core, FTypenameId(GUE.Names.MakeParametricTypeId(GUE.Typenames.Map, Params)) };
-		FStructId Id = GUE.Names.IndexStruct(Type);
+		FBindId Id = GUE.Names.IndexBindId(Type);
 		FMemberId MemberIds[] = { GUE.Members.Assign, GUE.Members.Remove, GUE.Members.Insert };
 
 		// Todo: Ownership / memory leak
-		const FStructDeclaration& Declaration = GUE.Types.DeclareStruct(Id, Type, 0, MemberIds, EMemberPresence::AllowSparse);
+		const FStructDeclaration& Declaration = GUE.Types.DeclareStruct(LowerCast(Id), Type, 0, MemberIds, EMemberPresence::AllowSparse);
 		auto Out = new TMapPropertyCustomBinding<KeyKind, ValueKind>(Property, Members);
-		GUE.Customs.BindStruct(Declaration.Id, *Out, Declaration, {});
+		GUE.Customs.BindStruct(Id, *Out, Declaration, {});
 
 		*OutCustomId = Id;
 		return Out;
 	}
 
 	template<EPropertyKind KeyKind>
-	inline IItemRangeBinding* New2(FMapProperty* Property, FMapMemberBindings Members, FStructId* OutCustomId)
+	inline IItemRangeBinding* New2(FMapProperty* Property, FMapMemberBindings Members, FBindId* OutCustomId)
 	{
 		switch (GetPropertyKind(Members.Value))
 		{
@@ -2634,7 +2634,7 @@ class FMapBindings
 		}
 	}
 
-	inline IItemRangeBinding* New(FMapProperty* Property, FMapMemberBindings Members, FStructId* OutCustomId)
+	inline IItemRangeBinding* New(FMapProperty* Property, FMapMemberBindings Members, FBindId* OutCustomId)
 	{
 		switch (GetPropertyKind(Members.Key))
 		{
@@ -2652,19 +2652,19 @@ class FMapBindings
 	}
 
 public:
-	FStructId BindNormal(FMapProperty* Property, FStructId PairId, FMapMemberBindings Members)
+	FBindId BindNormal(FMapProperty* Property, FBindId PairId, FMapMemberBindings Members)
 	{
-		if (const FStructId* CustomId = NormalBindings.Find(PairId))
+		if (const FBindId* CustomId = NormalBindings.Find(PairId))
 		{
 			return *CustomId;
 		}
 
-		FStructId CustomId;
+		FBindId CustomId;
 		New(Property, Members, /* out */ &CustomId);
 		return NormalBindings.Emplace(PairId, CustomId);
 	}
 
-	FRangeBinding BindFreezable(FMapProperty* Property, FStructId PairId, FMapMemberBindings Members)
+	FRangeBinding BindFreezable(FMapProperty* Property, FBindId PairId, FMapMemberBindings Members)
 	{
 		if (const FRangeBinding* RangeBinding = FrozenBindings.Find(PairId))
 		{
@@ -2688,25 +2688,25 @@ class FPairBindings
 		inline bool operator==(const FPair& O) const { return EqSkipOffset(KV[0], O.KV[0]) && EqSkipOffset(KV[1], O.KV[1]); }
 	};
 
-	TMap<FPair, FStructId>	Bindings;
+	TMap<FPair, FBindId>	Bindings;
 
-	inline FStructId BindImpl(FPair Pair)
+	inline FBindId BindImpl(FPair Pair)
 	{
 		check(Pair.KV[0].Offset == 0 && Pair.KV[1].Offset > 0);
-		if (const FStructId* BindId = Bindings.Find(Pair))
+		if (const FBindId* BindId = Bindings.Find(Pair))
 		{
 			return *BindId;
 		}
 
 		// Index names, can be optimized by checking if KeyParam / BindParam IsLowered()
-		FDualType KeyParam = Pair.KV[0].IndexParameterName(GUE.Names);
-		FDualType ValueParam = Pair.KV[1].IndexParameterName(GUE.Names);
+		FBothType KeyParam = Pair.KV[0].IndexParameterName(GUE.Names);
+		FBothType ValueParam = Pair.KV[1].IndexParameterName(GUE.Names);
 		FType BindParams[2] = { KeyParam.BindType, ValueParam.BindType };
 		FType DeclParams[2] = { KeyParam.DeclType, ValueParam.DeclType };
 		FType BindType = { GUE.Scopes.Core, FTypenameId(GUE.Names.MakeParametricTypeId(GUE.Typenames.Pair, BindParams)) };
 		FType DeclType = { GUE.Scopes.Core, FTypenameId(GUE.Names.MakeParametricTypeId(GUE.Typenames.Pair, DeclParams)) };
-		FStructId BindId = GUE.Names.IndexStruct(BindType);
-		FStructId DeclId = GUE.Names.IndexStruct(DeclType);
+		FBindId BindId = GUE.Names.IndexBindId(BindType);
+		FDeclId DeclId = GUE.Names.IndexDeclId(DeclType);
 		FMemberId Members[2] = { GUE.Members.Key, GUE.Members.Value };
 
 		// Todo: Ownership / memory leak
@@ -2968,7 +2968,7 @@ struct FStaticArrayBinding : IItemRangeBinding
 class FPropertyBinder
 {
 public:
-	FPropertyBinder(FStructId Id) : Owner(Id) {}
+	FPropertyBinder(FBindId Id) : Owner(Id) {}
 
 	void								AddMember(FMemberBinding Member)	{ Members.Add(Member); }
 	void								BindMember(FProperty* Property)		{ Members.Emplace(BindProperty(Property)); }
@@ -2978,7 +2978,7 @@ public:
 	bool								IsDense() { return GetOwnerOccupancy() == EMemberPresence::RequireAll; }
 
 private:
-	const FStructId							Owner;
+	const FBindId									Owner;
 	TOptional<FScopeId>								OwnerScope;
 	TOptional<EMemberPresence>						OwnerOccupancy;
 	TPagedArray<FRangeBinding, 1024>				Ranges;
@@ -3032,7 +3032,7 @@ private:
 	{
 		if (!OwnerOccupancy)
 		{
-			OwnerOccupancy = GUE.Types.Get(Owner).Occupancy;
+			OwnerOccupancy = GUE.Types.Get(LowerCast(Owner)).Occupancy;
 		}
 
 		return OwnerOccupancy.GetValue();
@@ -3054,7 +3054,7 @@ private:
 		return Out;
 	}
 
-	inline FMemberBinding BindAsStruct(FProperty* Property, FStructId Id)
+	inline FMemberBinding BindAsStruct(FProperty* Property, FBindId Id)
 	{
 		FMemberBinding Out(Property->GetOffset_ForInternal());
 		Out.InnermostSchema = FInnerId(Id);
@@ -3065,7 +3065,7 @@ private:
 	inline FMemberBinding BindAsStruct(FProperty* Property, UStruct* Struct)
 	{
 		FType Type = IndexType(SkipEmptyBases(Struct));
-		return BindAsStruct(Property, GUE.Names.IndexStruct(Type));
+		return BindAsStruct(Property, GUE.Names.IndexBindId(Type));
 	}
 
 	inline static FBitfieldBoolBindType MakeBitfieldBool(uint8 BitIdx)
@@ -3132,7 +3132,7 @@ private:
 		FMemberBinding Key = BindSingleProperty(Property->KeyProp);
 		FMemberBinding Value = BindSingleProperty(Property->ValueProp);
 		FMemberBinding Pair = GPairs.Bind(Key, Value);
-		FStructId PairId = Pair.InnermostSchema.Get().AsStruct();
+		FBindId PairId = Pair.InnermostSchema.Get().AsStructBindId();
 
 		bool bFreezable = EnumHasAnyFlags(Property->MapFlags, EMapPropertyFlags::UsesMemoryImageAllocator);
 		return bFreezable	? BindAsRange(Property, GMaps.BindFreezable(Property, PairId, {Key, Value, Pair}), Pair)
@@ -3196,7 +3196,7 @@ private:
 		}
 		else if (HasAny<MiscMask>(Flags))
 		{
-			FStructId BindId = HasAny<CASTCLASS_FInterfaceProperty>(Flags)
+			FBindId BindId = HasAny<CASTCLASS_FInterfaceProperty>(Flags)
 				? BindInterface(static_cast<FInterfaceProperty*>(Property))
 				: BindSparseDelegate(Owner, static_cast<FMulticastSparseDelegateProperty*>(Property));
 			return BindAsStruct(Property, BindId);
@@ -3268,11 +3268,11 @@ private:
 		{
 			// Create struct binding to allow delta serialization
 			FType StaticArrayType = MakeStaticArrayTypename(Property->GetFName());
-			FStructId StaticArrayId = GUE.Names.IndexStruct(StaticArrayType);
+			FBindId StaticArrayId = GUE.Names.IndexBindId(StaticArrayType);
 
 			// Todo: Ownership
 			TConstArrayView<FMemberId> Numerals = GUE.Numerals.MakeRange(IntCastChecked<uint16>(ArrayDim));
-			GUE.Types.DeclareNumeralStruct(StaticArrayId, StaticArrayType, Numerals, Occupancy);
+			GUE.Types.DeclareNumeralStruct(LowerCast(StaticArrayId), StaticArrayType, Numerals, Occupancy);
 
 			TArray<FMemberBinding, TInlineAllocator<64>> Elements;
 			Elements.Init(Out, ArrayDim);
@@ -3282,7 +3282,7 @@ private:
 				Element.Offset = Offset;
 				Offset += ElementSize;
 			}
-			GUE.Schemas.BindStruct(StaticArrayId, StaticArrayId, Elements);
+			GUE.Schemas.BindStruct(StaticArrayId, LowerCast(StaticArrayId), Elements);
 			
 			Out.InnermostType = DefaultStructBindType;
 			Out.InnermostSchema = FInnerId(StaticArrayId);
@@ -3317,7 +3317,7 @@ void BindSuperMembers(FPropertyBinder& Out, const UStruct* Struct)
 	}
 }
 
-void BindStruct(FStructId Id, const UStruct* Struct)
+void BindStruct(FBindId Id, const UStruct* Struct)
 {
 	if (GUE.Customs.FindStruct(Id))
 	{
@@ -3333,7 +3333,7 @@ void BindStruct(FStructId Id, const UStruct* Struct)
 		}
 		else
 		{
-			const FStructDeclaration& Declared = GUE.Types.Get(Id);
+			const FStructDeclaration& Declared = GUE.Types.Get(LowerCast(Id));
 			if (Declared.Super)
 			{
 				FMemberBinding Member;
@@ -3345,7 +3345,7 @@ void BindStruct(FStructId Id, const UStruct* Struct)
 	}
 	BindMembers(/* out */ Binder, Struct);
 
-	GUE.Schemas.BindStruct(Id, Id, Binder.GetMembers());
+	GUE.Schemas.BindStruct(Id, LowerCast(Id), Binder.GetMembers());
 	
 	// Don't bind CDOs, object defaults are passed in from top and objects aren't owned by containers 
 	if (Struct->HasAnyCastFlags(CASTCLASS_UScriptStruct))
@@ -3361,8 +3361,8 @@ static void BindInitialTypes()
 	UE_LOGFMT(LogPlainPropsUObject, Display, "Binding types to PlainProps schemas...");
 
 	// Declare all UScriptStruct/UClass/UFunction/UEnum
-	static const FStructId SkipStructs[] = { GUE.Structs.VerseFunction };
-	TArray<FStructId> Ids;
+	static const FBindId SkipStructs[] = { GUE.Structs.VerseFunction };
+	TArray<FBindId> Ids;
 	for (TObjectIterator<UField> It; It; ++It)
 	{
 		UField* Field = *It;
@@ -3370,11 +3370,11 @@ static void BindInitialTypes()
 		{
 			FType Type = IndexType(Struct);
 			FStructId Id = GUE.Names.IndexStruct(Type);
-			if (!Algo::Find(SkipStructs, Id))
+			if (!Algo::Find(SkipStructs, FBindId(Id)))
 			{
-				DeclareStruct(Struct, Type, Id);
+				DeclareStruct(Struct, Type, FDeclId(Id));
 			}
-			Ids.Add(Id);
+			Ids.Add(FBindId(Id));
 		}
 		else if (UEnum* Enum = Cast<UEnum>(Field))
 		{
@@ -3383,7 +3383,7 @@ static void BindInitialTypes()
 	}
 
 	// Bind all UScriptStruct/UClass/UFunction
-	const FStructId* IdIt = &Ids[0];
+	const FBindId* IdIt = &Ids[0];
 	for (TObjectIterator<UStruct> It; It; ++It, ++IdIt)
 	{
 		if (!Algo::Find(SkipStructs, *IdIt))
@@ -3394,7 +3394,7 @@ static void BindInitialTypes()
 	check(IdIt == &Ids[0] + Ids.Num());
 }
 
-static void DeclareProperty(FStructId Id, TConstArrayView<FMemberId> Members, EMemberPresence Occupancy)
+static void DeclareProperty(FDeclId Id, TConstArrayView<FMemberId> Members, EMemberPresence Occupancy)
 {
 	GUE.Types.DeclareStruct(Id,	GUE.Names.Resolve(Id), 0, Members, Occupancy);
 }
@@ -3567,9 +3567,9 @@ struct TCustomPropertyBindings
 		Bind(GUE.Structs.LazyObjectPtr, LazyObjectPtr);
 	}
 
-	void Bind(FStructId Id, ICustomBinding& Binding)
+	void Bind(FBindId Id, ICustomBinding& Binding)
 	{
-		Overlay.BindStruct(Id, Binding, GUE.Types.Get(Id));
+		Overlay.BindStruct(Id, Binding, GUE.Types.Get(LowerCast(Id)));
 	}
 
 	FCustomBindingsOverlay								Overlay;
@@ -3586,7 +3586,7 @@ struct TCustomPropertyBindings
 struct FMemoryBatch
 {
 	TArray64<uint8>				Data;
-	TArray<FStructId>		RuntimeIds; // To avoid reindexing schema FType
+	TArray<FStructId>			RuntimeIds; // To avoid reindexing schema FType
 	FMemoryPropertyBatch		Properties; // Contains FTexts referenced from Data
 };
 
@@ -3603,7 +3603,7 @@ public:
 		SavedObjects.Reserve(NumReserve);
 	}
 
-	FORCENOINLINE void Save(FStructId Id, const UObject* Object, const UObject* Arch)
+	FORCENOINLINE void Save(FBindId Id, const UObject* Object, const UObject* Arch)
 	{
 		FBuiltStruct* Built = Arch ? SaveStructDelta(Object, Arch, Id, DeltaCtx) : SaveStruct(Object, Id, FlatCtx);
 		SavedObjects.Emplace(Id, Built, Object);
@@ -3676,7 +3676,7 @@ public:
 private:
 	struct FSavedObject
 	{
-		FStructId		Id;
+		FBindId				Id;
 		FBuiltStruct*		Built;
 		const UObject*		Input; // For debug
 	};
@@ -3868,7 +3868,7 @@ struct FInstance
 	UObject* PP;
 	UObject* TPS;
 	UObject* UPS;
-	FStructId Id;
+	FBindId Id;
 
 	void Init()
 	{
@@ -3876,7 +3876,7 @@ struct FInstance
 		Class->GetDefaultObject(/* create lazily */ true);
 		Arch = Orig->GetArchetype();
 		check(Arch != Orig);
-		Id = GUE.Names.IndexStruct(IndexType(Class));
+		Id = GUE.Names.IndexBindId(IndexType(Class));
 	}
 };
 
diff --git a/Engine/Plugins/Experimental/PlainPropsUObject/Source/Public/PlainPropsUObjectRuntime.h b/Engine/Plugins/Experimental/PlainPropsUObject/Source/Public/PlainPropsUObjectRuntime.h
--- a/Engine/Plugins/Experimental/PlainPropsUObject/Source/Public/PlainPropsUObjectRuntime.h
+++ b/Engine/Plugins/Experimental/PlainPropsUObject/Source/Public/PlainPropsUObjectRuntime.h
@@ -25,16 +25,16 @@ class FDefaultStructs final : public IDefaultStructs
 {
 public:
 	~FDefaultStructs();
-	void									Bind(FStructId Id, const UScriptStruct* Struct);
-	void									BindZeroes(FStructId Id, SIZE_T Size, uint32 Alignment);
-	void									BindStatic(FStructId Id, const void* Struct);
-	void									Drop(FStructId Id);
-	virtual const void*						Get(FStructId Id) override;
+	void									Bind(FBindId Id, const UScriptStruct* Struct);
+	void									BindZeroes(FBindId Id, SIZE_T Size, uint32 Alignment);
+	void									BindStatic(FBindId Id, const void* Struct);
+	void									Drop(FBindId Id);
+	virtual const void*						Get(FBindId Id) override;
 
 private:
 	FMutableMemoryView						Zeroes;
 	TBitArray<>								Instanced;
-	TMap<FStructId, FDefaultInstance>		Instances;
+	TMap<FBindId, FDefaultInstance>			Instances;
 #if DO_CHECK
 	TBitArray<>								Bound;
 #endif
@@ -95,24 +95,24 @@ struct FCommonStructIds
 {
 	explicit FCommonStructIds(const FCommonScopeIds& Scopes, TIdIndexer<FName>& Names);
 	
-	FStructId					Name;
-	FStructId					Text;
-	FStructId					Guid;
-	FStructId					FieldPath;
-	FStructId					SoftObjectPath;
-	FStructId					ClassPtr;
-	FStructId					ObjectPtr;
-	FStructId					WeakObjectPtr;
-	FStructId					LazyObjectPtr;
-	FStructId					SoftObjectPtr;
-	FStructId					ScriptInterface;
-	FStructId					Delegate;
-	FStructId					MulticastDelegate;
-	FStructId					MulticastInlineDelegate;
-	FStructId					MulticastSparseDelegate;
-	FStructId					VerseFunction;
-	FStructId					DynamicallyTypedValue;
-	FStructId					ReferencePropertyValue;
+	FDualStructId			Name;
+	FDualStructId			Text;
+	FDualStructId			Guid;
+	FDualStructId			FieldPath;
+	FDualStructId			SoftObjectPath;
+	FDualStructId			ClassPtr;
+	FDualStructId			ObjectPtr;
+	FDualStructId			WeakObjectPtr;
+	FDualStructId			LazyObjectPtr;
+	FDualStructId			SoftObjectPtr;
+	FDualStructId			ScriptInterface;
+	FDualStructId			Delegate;
+	FDeclId					MulticastDelegate;
+	FBindId					MulticastInlineDelegate;
+	FDualStructId			MulticastSparseDelegate;
+	FDualStructId			VerseFunction;
+	FDualStructId			DynamicallyTypedValue;
+	FDualStructId			ReferencePropertyValue;
 };
 
 struct FCommonMemberIds
