PlainProps diff / delta improvements

* Generic static custom bindings now respect DiffCustom rather than calling native !=
* Mechanism for ICustomBinding to capture which member differ
* Templated DiffMember support in TMemberSerializer

#rnx

diff --git a/Engine/Plugins/Experimental/PlainProps/Source/Private/PlainPropsBind.cpp b/Engine/Plugins/Experimental/PlainProps/Source/Private/PlainPropsBind.cpp
--- a/Engine/Plugins/Experimental/PlainProps/Source/Private/PlainPropsBind.cpp
+++ b/Engine/Plugins/Experimental/PlainProps/Source/Private/PlainPropsBind.cpp
@@ -1,6 +1,7 @@
 // Copyright Epic Games, Inc. All Rights Reserved.
 
 #include "PlainPropsBind.h"
+#include "PlainPropsDiff.h"
 #include "PlainPropsIndex.h"
 #include "PlainPropsInternalBind.h"
 #include "PlainPropsInternalBuild.h"
@@ -32,6 +33,14 @@ uint32 FSchemaBinding::CalculateSize() const
 
 ////////////////////////////////////////////////////////////////////////////////////////////////
 
+bool ICustomBinding::DiffCustom(const void* StructA, const void* StructB, FDiffContext& Ctx) const
+{
+	// Use faster non-tracking diff used by delta saving by default
+	return DiffCustom(StructA, StructB, static_cast<const FBindContext&>(Ctx));
+}
+
+////////////////////////////////////////////////////////////////////////////////////////////////
+
 FMemberVisitor::FMemberVisitor(const FSchemaBinding& InSchema)
 : Schema(InSchema)
 , NumMembers(InSchema.NumMembers)
diff --git a/Engine/Plugins/Experimental/PlainProps/Source/Private/PlainPropsDiff.cpp b/Engine/Plugins/Experimental/PlainProps/Source/Private/PlainPropsDiff.cpp
--- a/Engine/Plugins/Experimental/PlainProps/Source/Private/PlainPropsDiff.cpp
+++ b/Engine/Plugins/Experimental/PlainProps/Source/Private/PlainPropsDiff.cpp
@@ -11,10 +11,11 @@
 namespace PlainProps
 {
 
-static bool DiffItem(const void* A, const void* B, FDiffContext& Ctx, FRangeMemberBinding Range);
-static bool DiffItem(const void* A, const void* B, FDiffContext& Ctx, FBindId Struct);
-
-static bool DiffMember(FLeafMemberBinding Member, const void* OwnerA, const void* OwnerB, FDiffContext& Ctx, FOptionalMemberId Name)
+inline bool DiffMember(FLeafMemberBinding Member, const void* OwnerA, const void* OwnerB, const FBindContext&, FOptionalMemberId)
+{
+	return !!DiffLeaf(At(OwnerA, Member.Offset), At(OwnerB, Member.Offset), Member.Leaf);
+}
+inline bool DiffMember(FLeafMemberBinding Member, const void* OwnerA, const void* OwnerB, FDiffContext& Ctx, FOptionalMemberId Name)
 {
 	const uint8* A = At(OwnerA, Member.Offset);
 	const uint8* B = At(OwnerB, Member.Offset);
@@ -26,7 +27,8 @@ static bool DiffMember(FLeafMemberBinding Member, const void* OwnerA, const void
 	return false;
 }
 
-static bool DiffMember(FRangeMemberBinding Member, const void* OwnerA, const void* OwnerB, FDiffContext& Ctx, FOptionalMemberId Name)
+template<typename ContextType>
+bool DiffMember(FRangeMemberBinding Member, const void* OwnerA, const void* OwnerB, ContextType& Ctx, FOptionalMemberId Name)
 {
 	FRangeBinding Binding = Member.RangeBindings[0];
 	FMemberBindType InnerType = Member.InnerTypes[0];
@@ -48,6 +50,7 @@ static bool DiffMember(FRangeMemberBinding Member, const void* OwnerA, const voi
 		}
 		else if (InnerKind == EMemberKind::Leaf)
 		{
+			// Todo: float/double ranges should use PreciseFP
 			bDiff = DiffItemRange(A, B, ItemBinding, Ctx, SizeOf(GetItemWidth(InnerType.AsLeaf())));
 		}
 		else
@@ -55,24 +58,35 @@ static bool DiffMember(FRangeMemberBinding Member, const void* OwnerA, const voi
 			bDiff = DiffItemRange(A, B, ItemBinding, Ctx, GetInnerRange(Member));
 		}
 	}
-	if (bDiff)
+
+	if constexpr (std::is_same_v<ContextType, FDiffContext>)
 	{
-		Ctx.Out.Add({FMemberBindType(Binding.GetSizeType()), Name, {.Range = Binding}, A, B});
+		if (bDiff)
+		{
+			Ctx.Out.Add({FMemberBindType(Binding.GetSizeType()), Name, {.Range = Binding}, A, B});
+		}
 	}
 	return bDiff;
 }
 
-static bool DiffMembers(FBindId Id, const void* A, const void* B, FDiffContext& Ctx);
+template<typename ContextType>
+bool DiffMembers(FBindId Id, const void* A, const void* B, ContextType& Ctx);
 
-inline bool DiffStruct(FBindId Id, const void* A, const void* B, FDiffContext& Ctx)
+template<typename ContextType>
+inline bool DiffStruct(FBindId Id, const void* A, const void* B, ContextType& Ctx)
 {
 	if (const ICustomBinding* Custom = Ctx.Customs.FindStruct(Id))
 	{
-		return Custom->DiffCustom(A, B);
+		return Custom->DiffCustom(A, B, Ctx);
 	}
 	return DiffMembers(Id, A, B, Ctx);
 }
 
+static bool DiffMember(FStructMemberBinding Member, const void* OwnerA, const void* OwnerB, const FBindContext& Ctx, FMemberId)
+{
+	return DiffStruct(Member.Id, At(OwnerA, Member.Offset), At(OwnerB, Member.Offset), Ctx);
+}
+
 static bool DiffMember(FStructMemberBinding Member, const void* OwnerA, const void* OwnerB, FDiffContext& Ctx, FMemberId Name)
 {
 	const uint8* A = At(OwnerA, Member.Offset);
@@ -85,7 +99,8 @@ static bool DiffMember(FStructMemberBinding Member, const void* OwnerA, const vo
 	return false;
 }
 
-static bool DiffMember(/* in-out */ FMemberVisitor& It, const void* A, const void* B, FDiffContext& Ctx, FMemberId Name)
+template<typename ContextType>
+bool DiffMember(/* in-out */ FMemberVisitor& It, const void* A, const void* B, ContextType& Ctx, FMemberId Name)
 {
 	switch (It.PeekKind())
 	{
@@ -96,11 +111,11 @@ static bool DiffMember(/* in-out */ FMemberVisitor& It, const void* A, const voi
 	return true;
 }
 
-
-static bool DiffMembers(FBindId Id, const void* A, const void* B, FDiffContext& Ctx)
+template<typename ContextType>
+bool DiffMembers(FBindId Id, const void* A, const void* B, ContextType& Ctx)
 {
 	const FSchemaBinding& Schema = Ctx.Schemas.GetStruct(Id);
-	const FStructDeclaration& Declaration = Ctx.Types.Get(Schema.DeclId);
+	const FStructDeclaration& Declaration = Ctx.Declarations.Get(Schema.DeclId);
 	
 	FMemberVisitor It(Schema);
 	if (Schema.HasSuper())
@@ -123,26 +138,64 @@ static bool DiffMembers(FBindId Id, const void* A, const void* B, FDiffContext&
 	return false;
 }
 
-static bool DiffItem(const void* A, const void* B, FDiffContext& Ctx, FRangeMemberBinding Range)
+template<typename ContextType>
+bool DiffItem(const void* A, const void* B, ContextType& Ctx, FRangeMemberBinding Range)
 {
 	return DiffMember(Range, A, B, Ctx, NoId);
 }
 
-static bool DiffItem(const void* A, const void* B, FDiffContext& Ctx, FBindId Id)
+template<typename ContextType>
+bool DiffItem(const void* A, const void* B, ContextType& Ctx, FBindId Id)
 {
 	return DiffStruct(Id, A, B, Ctx);
 }
 
 ////////////////////////////////////////////////////////////////////////////////////////////////
 
-bool DiffStructs(FBindId Id, const void* A, const void* B, FDiffContext& Ctx)
+bool DiffStructs(const void* A, const void* B, FBindId Id, const FBindContext& Ctx)
 {
-	if (const ICustomBinding* Custom = Ctx.Customs.FindStruct(Id))
-	{
-		// Currently no way for DiffCustom to emit deeper diff into
-		return Custom->DiffCustom(A, B);
-	}
-	return DiffMembers(Id, A, B, Ctx);
+	return DiffStruct(Id, A, B, Ctx);
+}
+
+bool DiffStructs(const void* A, const void* B, FBindId Id, FDiffContext& Ctx)
+{
+	return DiffStruct(Id, A, B, Ctx);
+}
+
+bool DiffLeaves(float A, float B)
+{
+	return !::UE::PreciseFPEqual(A, B);
+}
+
+bool DiffLeaves(double A, double B)
+{
+	return !::UE::PreciseFPEqual(A, B);
+}
+
+bool DiffRanges(const void* A, const void* B, const IItemRangeBinding& Binding, FUnpackedLeafType Leaf)
+{
+	// Todo: float/double ranges should use PreciseFP
+	return DiffItemRange(A, B, Binding, /* dummy ctx */ Binding, SizeOf(Leaf.Width));
+}
+
+bool DiffRanges(const void* A, const void* B, const IItemRangeBinding& Binding, FBindId ItemType, const FBindContext& Ctx)
+{
+	return DiffItemRange(A, B, Binding, Ctx, ItemType);
+}
+
+bool DiffRanges(const void* A, const void* B, const IItemRangeBinding& Binding, FBindId ItemType, FDiffContext& Ctx)
+{
+	return DiffItemRange(A, B, Binding, Ctx, ItemType);
+}
+
+bool DiffRanges(const void* A, const void* B, const IItemRangeBinding& Binding, FRangeMemberBinding ItemType, const FBindContext& Ctx)
+{
+	return DiffItemRange(A, B, Binding, Ctx, ItemType);
+}
+
+bool DiffRanges(const void* A, const void* B, const IItemRangeBinding& Binding, FRangeMemberBinding ItemType, FDiffContext& Ctx)
+{
+	return DiffItemRange(A, B, Binding, Ctx, ItemType);
 }
 
 ////////////////////////////////////////////////////////////////////////////////
diff --git a/Engine/Plugins/Experimental/PlainProps/Source/Private/PlainPropsLoad.cpp b/Engine/Plugins/Experimental/PlainProps/Source/Private/PlainPropsLoad.cpp
--- a/Engine/Plugins/Experimental/PlainProps/Source/Private/PlainPropsLoad.cpp
+++ b/Engine/Plugins/Experimental/PlainProps/Source/Private/PlainPropsLoad.cpp
@@ -488,7 +488,7 @@ template<class LoadIdBinder>
 struct FLoadBindingBase : ICustomBinding
 {
 	virtual void SaveCustom(FMemberBuilder& Dst, const void* Src, const void* Default, const FSaveContext& Ctx) override final { unimplemented(); }
-	virtual bool DiffCustom(const void*, const void*) const override final { unimplemented(); return false; }
+	virtual bool DiffCustom(const void*, const void*, const FBindContext&) const override final { unimplemented(); return false; }
 };
 struct FTypeErasedLoadBinding final : FLoadBindingBase
 {
diff --git a/Engine/Plugins/Experimental/PlainProps/Source/Private/PlainPropsSave.cpp b/Engine/Plugins/Experimental/PlainProps/Source/Private/PlainPropsSave.cpp
--- a/Engine/Plugins/Experimental/PlainProps/Source/Private/PlainPropsSave.cpp
+++ b/Engine/Plugins/Experimental/PlainProps/Source/Private/PlainPropsSave.cpp
@@ -389,7 +389,7 @@ static bool DiffItem(const uint8* A, const uint8* B, const FSaveContext& Ctx, FB
 {
 	if (const ICustomBinding* Custom = Ctx.Customs.FindStruct(Id))
 	{
-		return Custom->DiffCustom(A, B);
+		return Custom->DiffCustom(A, B, Ctx);
 	}
 	
 	bool bOut = false;
@@ -478,7 +478,7 @@ FBuiltStruct* SaveStructDelta(const void* Struct, const void* Default, FBindId I
 	const FStructDeclaration* Declaration = nullptr;
 	if (ICustomBinding* Custom = Ctx.Customs.FindStructToSave(Id, /* out */ Declaration))
 	{
-		if (Custom->DiffCustom(Struct, Default))
+		if (Custom->DiffCustom(Struct, Default, Ctx))
 		{
 			Custom->SaveCustom(Out, Struct, Default, Ctx);
 		}
@@ -497,7 +497,7 @@ FBuiltStruct* SaveStructDeltaIfDiff(const void* Struct, const void* Default, FBi
 	const FStructDeclaration* Declaration = nullptr;
 	if (ICustomBinding* Custom = Ctx.Customs.FindStructToSave(Id, /* out */ Declaration))
 	{
-		if (Custom->DiffCustom(Struct, Default))
+		if (Custom->DiffCustom(Struct, Default, Ctx))
 		{
 			Custom->SaveCustom(Out, Struct, Default, Ctx);
 			return Out.BuildAndReset(Ctx.Scratch, *Declaration, Ctx.Declarations.GetDebug());
diff --git a/Engine/Plugins/Experimental/PlainProps/Source/Private/PlainPropsUeCoreBindings.cpp b/Engine/Plugins/Experimental/PlainProps/Source/Private/PlainPropsUeCoreBindings.cpp
--- a/Engine/Plugins/Experimental/PlainProps/Source/Private/PlainPropsUeCoreBindings.cpp
+++ b/Engine/Plugins/Experimental/PlainProps/Source/Private/PlainPropsUeCoreBindings.cpp
@@ -2,6 +2,7 @@
 
 #include "PlainPropsUeCoreBindings.h"
 #include "PlainPropsBuild.h"
+#include "PlainPropsDiff.h"
 #include "PlainPropsSave.h"
 #include "Math/Transform.h"
 
@@ -45,6 +46,32 @@ void FTransformBinding::Save(FMemberBuilder& Dst, const FTransform& Src, const F
 	}
 }
 
+inline FDiffNode MakeStructDiff(const void* A, const void* B, FMemberId Name, FBindId Id)
+{
+	return { DefaultStructBindType, Name, {.Struct = Id}, A, B };
+}
+
+bool FTransformBinding::Diff(const FTransform& A, const FTransform& B, FDiffContext& Ctx) const
+{
+	if (!A.TranslationEquals(B, 0.0))
+	{
+		Ctx.Out.Emplace(MakeStructDiff(&A, &B, MemberIds[(uint8)EMember::Translate], VectorId));
+		return true;
+	}
+	else if (!A.RotationEquals(B, 0.0))
+	{
+		Ctx.Out.Emplace(MakeStructDiff(&A, &B, MemberIds[(uint8)EMember::Rotate], QuatId));
+		return true;
+	}
+	else if (!A.Scale3DEquals(B, 0.0))
+	{
+		Ctx.Out.Emplace(MakeStructDiff(&A, &B, MemberIds[(uint8)EMember::Scale], VectorId));
+		return true;
+	}
+
+	return false;
+}
+
 template<typename T>
 T GrabAndMemcpy(FMemberLoader& Members)
 {
diff --git a/Engine/Plugins/Experimental/PlainProps/Source/Public/PlainPropsBind.h b/Engine/Plugins/Experimental/PlainProps/Source/Public/PlainPropsBind.h
--- a/Engine/Plugins/Experimental/PlainProps/Source/Public/PlainPropsBind.h
+++ b/Engine/Plugins/Experimental/PlainProps/Source/Public/PlainPropsBind.h
@@ -17,9 +17,11 @@
 namespace PlainProps 
 {
 
+struct FBindContext;
 struct FBuiltRange;
 struct FBuiltStruct;
 struct FCustomBindingHandle;
+struct FDiffContext;
 class FIdIndexerBase;
 struct FLoadBatch;
 class FMemberBuilder;
@@ -259,13 +261,16 @@ enum ECustomLoadMethod { Construct, Assign };
 /// * non-default constructible types
 /// * custom delta semantics
 /// * other runtime representations than struct/class, e.g. serialize database
-/// * optimization for very common struct
+/// * optimizations of very common structs
 struct ICustomBinding
 {
 	virtual ~ICustomBinding() {}
 	virtual void	SaveCustom(FMemberBuilder& Dst, const void* Src, const void* Default, const FSaveContext& Ctx) = 0;
 	virtual void	LoadCustom(void* Dst, FStructLoadView Src, ECustomLoadMethod Method) const = 0;
-	virtual bool	DiffCustom(const void* StructA, const void* StructB) const = 0;
+	virtual bool	DiffCustom(const void* StructA, const void* StructB, const FBindContext& Ctx) const = 0;
+
+	// Overload to track the first diffing member, see FDiffPath
+	PLAINPROPS_API virtual bool	DiffCustom(const void* StructA, const void* StructB, FDiffContext& Ctx) const;
 };
 
 struct FInnersHandle
@@ -643,6 +648,15 @@ struct IDefaultStructs
 
 ////////////////////////////////////////////////////////////////////////////////////////////////
 
+struct FBindContext
+{
+	const FDeclarations&		Declarations;
+	const FSchemaBindings&		Schemas;
+	FCustomBindings&			Customs;
+};
+
+////////////////////////////////////////////////////////////////////////////////////////////////
+
 struct FStructBindIds final : IBindIds
 {
 	FStructBindIds(const FCustomBindings& InCustoms, const FSchemaBindings& InSchemas)
@@ -906,10 +920,15 @@ struct TScopedStructBinding : FBothStructId, CustomBinding
 	{
 		CustomBinding::Load(*static_cast<T*>(Dst), Src, Method);
 	}
+	
+	virtual bool DiffCustom(const void* A, const void* B, const FBindContext& Ctx) const override
+	{
+		return CustomBinding::Diff(*static_cast<const T*>(A), *static_cast<const T*>(B), Ctx);
+	}
 
-	virtual bool DiffCustom(const void* A, const void* B) const override
+	virtual bool DiffCustom(const void* A, const void* B, FDiffContext& Ctx) const override
 	{
-		return CustomBinding::Diff(*static_cast<const T*>(A), *static_cast<const T*>(B));
+		return CustomBinding::Diff(*static_cast<const T*>(A), *static_cast<const T*>(B), Ctx);
 	}
 };
 
diff --git a/Engine/Plugins/Experimental/PlainProps/Source/Public/PlainPropsDiff.h b/Engine/Plugins/Experimental/PlainProps/Source/Public/PlainPropsDiff.h
--- a/Engine/Plugins/Experimental/PlainProps/Source/Public/PlainPropsDiff.h
+++ b/Engine/Plugins/Experimental/PlainProps/Source/Public/PlainPropsDiff.h
@@ -11,7 +11,7 @@ union FDiffMetadata
 {
 	FOptionalEnumId			Leaf;
 	FRangeBinding			Range;
-	FStructId				Struct;
+	FBindId					Struct;
 };
 
 // Currently lacking range indices
@@ -26,15 +26,29 @@ struct FDiffNode
 
 struct FDiffPath : public TArray<FDiffNode, TInlineAllocator<16>> {};
 
-struct FDiffContext
+// Tracks diff path for diff tools unlike the const FBindContext& overloads for delta saving
+struct FDiffContext : FBindContext
 {
-	const FDeclarations&	Types;
-	const FSchemaBindings&	Schemas;
-	const FCustomBindings&	Customs;
-	FDiffPath				Out;
+	FDiffPath Out;
 };
 
-[[nodiscard]] PLAINPROPS_API bool DiffStructs(FBindId Id, const void* A, const void* B, FDiffContext& Ctx);
+////////////////////////////////////////////////////////////////////////////////
+// Tracking and non-tracking methods to diff member leaves/ranges/structs
+//
+// FDiffContext& overloads track inner FDiffPath, caller must add outermost FDiffNode
+
+[[nodiscard]] PLAINPROPS_API bool DiffStructs(const void* A, const void* B, FBindId Id, const FBindContext& Ctx);
+[[nodiscard]] PLAINPROPS_API bool DiffStructs(const void* A, const void* B, FBindId Id, FDiffContext& Ctx);
+
+[[nodiscard]] PLAINPROPS_API bool DiffLeaves(float A, float B);
+[[nodiscard]] PLAINPROPS_API bool DiffLeaves(double A, double B);
+template<Arithmetic T>		 bool DiffLeaves(T A, T B) { return A != B; }
+
+[[nodiscard]] PLAINPROPS_API bool DiffRanges(const void* A, const void* B, const IItemRangeBinding& Binding, FUnpackedLeafType ItemType);
+[[nodiscard]] PLAINPROPS_API bool DiffRanges(const void* A, const void* B, const IItemRangeBinding& Binding, FBindId ItemType, const FBindContext& Ctx);
+[[nodiscard]] PLAINPROPS_API bool DiffRanges(const void* A, const void* B, const IItemRangeBinding& Binding, FBindId ItemType, FDiffContext& Ctx);
+[[nodiscard]] PLAINPROPS_API bool DiffRanges(const void* A, const void* B, const IItemRangeBinding& Binding, FRangeMemberBinding ItemType, const FBindContext& Ctx);
+[[nodiscard]] PLAINPROPS_API bool DiffRanges(const void* A, const void* B, const IItemRangeBinding& Binding, FRangeMemberBinding ItemType, FDiffContext& Ctx);
 
 ////////////////////////////////////////////////////////////////////////////////
 
diff --git a/Engine/Plugins/Experimental/PlainProps/Source/Public/PlainPropsSave.h b/Engine/Plugins/Experimental/PlainProps/Source/Public/PlainPropsSave.h
--- a/Engine/Plugins/Experimental/PlainProps/Source/Public/PlainPropsSave.h
+++ b/Engine/Plugins/Experimental/PlainProps/Source/Public/PlainPropsSave.h
@@ -2,26 +2,13 @@
 
 #pragma once
 
-#include "PlainPropsTypes.h"
+#include "PlainPropsBind.h"
 
 namespace PlainProps 
 {
 
-struct FBuiltRange;
-struct FBuiltStruct;
-class FCustomBindings;
-class FDeclarations;
-struct IDefaultStructs;
-struct ILeafRangeBinding;
-struct FRangeMemberBinding;
-class FSchemaBindings;
-class FScratchAllocator;
-
-struct FSaveContext
+struct FSaveContext : FBindContext
 {
-	const FDeclarations&		Declarations;
-	const FSchemaBindings&		Schemas;
-	FCustomBindings&			Customs;
 	FScratchAllocator&			Scratch;
 	IDefaultStructs*			Defaults = nullptr;
 };
diff --git a/Engine/Plugins/Experimental/PlainProps/Source/Public/PlainPropsSaveMember.h b/Engine/Plugins/Experimental/PlainProps/Source/Public/PlainPropsSaveMember.h
--- a/Engine/Plugins/Experimental/PlainProps/Source/Public/PlainPropsSaveMember.h
+++ b/Engine/Plugins/Experimental/PlainProps/Source/Public/PlainPropsSaveMember.h
@@ -5,12 +5,28 @@
 #include "PlainPropsTypes.h"
 #include "PlainPropsBind.h"
 #include "PlainPropsBuild.h"
+#include "PlainPropsDiff.h"
 #include "PlainPropsLoadMember.h"
 #include "PlainPropsSave.h"
 
 namespace PlainProps 
 {
 
+FORCEINLINE bool Track(bool bDiff, const FBindContext&, FMemberBindType, FMemberId, FDiffMetadata, const void*, const void*)
+{
+	return bDiff;
+}
+
+FORCEINLINE bool Track(bool bDiff, FDiffContext& Ctx, FMemberBindType Type, FMemberId Name, FDiffMetadata Meta, const void* A, const void* B)
+{
+	if (bDiff)
+	{
+		Ctx.Out.Emplace(Type, ToOptional(Name), Meta, A, B);
+	}
+	
+	return bDiff;
+}
+
 // TMemberSerializer helper
 template<typename T>
 struct TStructSerializer
@@ -35,18 +51,24 @@ struct TStructSerializer
 		LoadStruct(&Dst, Src.GrabStruct());
 	}
 
-	void ConstructAndLoadMember(T& Dst, /* in-out */ FMemberLoader& Src) const
+	void ConstructAndLoadMember(void* Dst, /* in-out */ FMemberLoader& Src) const
 	{
 		if constexpr (std::is_default_constructible_v<T>)
 		{
-			LoadMember(*new (&Dst) T, Src);
+			LoadMember(*new (Dst) T, Src);
 		}
 		else
 		{
 			static_assert(!std::is_void_v<CustomBind<T>>, "Non-default constructible types must load via custom bindings");
-			ConstructAndLoadStruct(&Dst, Src.GrabStruct());
+			ConstructAndLoadStruct(Dst, Src.GrabStruct());
 		}
 	}
+	
+	template<class ContextType>
+	bool DiffMember(const T& A, const T& B, FMemberId Name, ContextType& Ctx) const
+	{
+		return Track(DiffStructs(&A, &B, Id, Ctx), Ctx, DefaultStructBindType, Name, {.Struct = Id}, &A, &B);
+	}
 
 	////////// TRangeSaver API //////////
 	
@@ -97,6 +119,7 @@ template<typename T, typename RangeBinding>
 struct TRangeSerializer
 {
 	static constexpr EMemberKind	Kind = EMemberKind::Range;
+	static constexpr bool			bLeafRange = std::is_base_of_v<ILeafRangeBinding, RangeBinding>;
 	static constexpr uint16			NumRanges = CountRangeBindings<RangeBinding>();
 	static constexpr uint16			NumTypes = NumRanges + 1;
 
@@ -123,7 +146,7 @@ struct TRangeSerializer
 
 		if constexpr (!LeafType<InnermostType>)
 		{			
-			FBothStructId Id = {InnermostSchema.Get().AsStruct(), GetStructDeclId<Ids, InnermostType>()};
+			FBothStructId Id = {InnermostSchema.Get().AsStructBindId(), GetStructDeclId<Ids, InnermostType>()};
 			Init.RegisterInnerStruct(Id, Names);
 		}
 	}
@@ -133,16 +156,68 @@ struct TRangeSerializer
 		LoadRange(&Dst, Src.GrabRange(), MakeArrayView(Bindings, NumRanges));
 	}
 
+	void ConstructAndLoadMember(void* Dst, /* in-out */ FMemberLoader& Src) const
+	{
+		static_assert(std::is_default_constructible_v<T>, TEXT("Ranges must be default-constructible"));
+		LoadMember(*new (Dst) T, Src);
+	}
+
 	FTypedValue	SaveMember(const T& Value, const FSaveContext& Ctx) const
 	{
 		return { MakeMemberSchema(), {.Range = SaveItem(Value, Ctx)} };
 	}
+
+	template<typename ContextType>
+	bool DiffItems(const void* A, const void* B, const IItemRangeBinding& Binding, ContextType& Ctx) const
+	{
+		if constexpr (Kind == EMemberKind::Leaf)
+		{
+			return DiffRanges(A, B, Binding, ReflectLeaf<T>);
+		}
+		else if constexpr (Kind == EMemberKind::Struct)
+		{
+			return DiffRanges(A, B, Binding, InnermostSchema.Get().AsStructBindId(), Ctx);
+		}
+		else
+		{
+			return DiffRanges(A, B, Binding, MakeInnerRangeBinding(), Ctx);
+		}
+	}
+
+	inline FDiffMetadata MakeDiffMetadata() const
+	{
+		if constexpr (Kind == EMemberKind::Leaf)
+		{
+			return { .Leaf = ToOptionalEnum(InnermostSchema) };
+		}
+		else if constexpr (Kind == EMemberKind::Struct)
+		{
+			return { .Struct = InnermostSchema.Get().AsStructBindId() };
+		}
+		else
+		{
+			return { .Range = Bindings[1] };
+		}
+	}
+
+	template<typename ContextType>
+	bool DiffMember(const T& A, const T& B, FMemberId Name, ContextType& Ctx) const
+	{
+		bool bDiff = bLeafRange ? Bindings[0].AsLeafBinding().DiffLeaves(A, B)
+			: DiffItems(&A, &B, Bindings[0].AsItemBinding(), Ctx);
+		return Track(bDiff, Ctx, Types[0], Name, MakeDiffMetadata(), &A, &B);
+	}
 		
 	FMemberSchema MakeMemberSchema() const
 	{
 		return { Types[0], Types[1], NumRanges, InnermostSchema, NumRanges > 1 ? Types + 1 : nullptr };
 	}
 
+	FRangeMemberBinding MakeInnerRangeBinding() const
+	{
+		return { &BindTypes[1].Value, Bindings, NumRanges, InnermostSchema, /* offset */ 0 };
+	}
+
 	////////// TRangeSaver API //////////
 
 	using BuiltItemType = FBuiltRange*;
@@ -154,15 +229,13 @@ struct TRangeSerializer
 
 	FBuiltRange* SaveItem(const T& Item, const FSaveContext& Ctx) const
 	{
-		if constexpr (std::is_base_of_v<ILeafRangeBinding, RangeBinding>)
+		if constexpr (bLeafRange)
 		{
-			static_assert(NumRanges == 1);
 			return SaveLeafRange(&Item, Bindings[0].AsLeafBinding(), Types[1].AsLeaf(), Ctx);
 		}
 		else
 		{
-			FRangeMemberBinding MemberBinding = { &BindTypes[1].Value, Bindings, NumRanges, InnermostSchema, /* offset */ 0 };
-			return SaveRange(&Item, MemberBinding, Ctx);
+			return SaveRange(&Item, MakeInnerRangeBinding(), Ctx);
 		}
 	}
 };
@@ -173,14 +246,14 @@ struct TRangeSerializer
 template<typename T>
 struct TLeafSerializer
 {
-	void LoadMember(T& Dst, /* in-out */ FMemberLoader& Src) const
+	static void LoadMember(T& Dst, /* in-out */ FMemberLoader& Src)
 	{
 		Dst = Src.GrabLeaf().As<T>();
 	}
 
-	void ConstructAndLoadMember(T& Dst, /* in-out */ FMemberLoader& Src) const
+	static void ConstructAndLoadMember(void* Dst, /* in-out */ FMemberLoader& Src)
 	{
-		LoadMember(Dst, Src);
+		LoadMember(*static_cast<T*>(Dst), Src);
 	}
 
 	static FTypedValue SaveLeaf(FMemberType Type, FOptionalInnerId Id, T Value)
@@ -200,6 +273,7 @@ template<Arithmetic T>
 struct TArithmeticSerializer : TLeafSerializer<T>
 {
 	static constexpr FMemberType MemberType = ReflectArithmetic<T>.Pack();
+	static constexpr FMemberBindType MemberBindType{ReflectArithmetic<T>};
 
 	template<class Ids>
 	TArithmeticSerializer(TCustomInit<Ids>, TConstArrayView<FMemberId>)
@@ -207,7 +281,13 @@ struct TArithmeticSerializer : TLeafSerializer<T>
 
 	static FTypedValue SaveMember(T Value, const FSaveContext&)
 	{
-		return SaveLeaf(MemberType, NoId, Value);
+		return TLeafSerializer<T>::SaveLeaf(MemberType, NoId, Value);
+	}
+
+	template<class ContextType>
+	inline static bool DiffMember(const T& A, const T& B, FMemberId Name, ContextType& Ctx)
+	{
+		return Track(A != B, Ctx, MemberBindType, Name, {.Leaf = NoId}, &A, &B);
 	}
 
 	////////// TRangeSaver API //////////
@@ -223,6 +303,7 @@ template<Enumeration T>
 struct TEnumSerializer : TLeafSerializer<T>
 {
 	static constexpr FMemberType MemberType = ReflectEnum<T>.Pack();
+	static constexpr FMemberBindType MemberBindType = ReflectEnum<T>;
 	
 	const FEnumId Id;
 
@@ -233,7 +314,13 @@ struct TEnumSerializer : TLeafSerializer<T>
 
 	FTypedValue SaveMember(T Value, const FSaveContext& Ctx) const
 	{
-		return SaveLeaf(MemberType, Id, Value);
+		return TLeafSerializer<T>::SaveLeaf(MemberType, Id, Value);
+	}
+
+	template<class ContextType>
+	inline bool DiffMember(const T& A, const T& B, FMemberId Name, ContextType& Ctx) const
+	{
+		return Track(A != B, Ctx, MemberBindType, Name, {.Leaf = Id}, &A, &B);
 	}
 
 	////////// TRangeSaver API //////////
diff --git a/Engine/Plugins/Experimental/PlainProps/Source/Public/PlainPropsUeCoreBindings.h b/Engine/Plugins/Experimental/PlainProps/Source/Public/PlainPropsUeCoreBindings.h
--- a/Engine/Plugins/Experimental/PlainProps/Source/Public/PlainPropsUeCoreBindings.h
+++ b/Engine/Plugins/Experimental/PlainProps/Source/Public/PlainPropsUeCoreBindings.h
@@ -468,8 +468,8 @@ struct TPairBinding : ICustomBinding
 		check(Members.PeekName() == MemberIds[0]);
 		if (Method == ECustomLoadMethod::Construct)
 		{
-			Key.ConstructAndLoadMember(/* out */ Dst.Key, /* in-out */ Members);
-			Value.ConstructAndLoadMember(/* out */ Dst.Value, /* in-out */ Members);		
+			Key.ConstructAndLoadMember(/* out */ &Dst.Key, /* in-out */ Members);
+			Value.ConstructAndLoadMember(/* out */ &Dst.Value, /* in-out */ Members);		
 		}
 		else
 		{
@@ -477,8 +477,13 @@ struct TPairBinding : ICustomBinding
 			Value.LoadMember(/* out */ Dst.Value, /* in-out */ Members);				
 		}
 	}
-	
-	static bool	Diff(const TPair<K,V>& A, const TPair<K,V>& B) { return A != B; }
+
+	template<typename ContextType>
+	bool Diff(const TPair<K,V>& A, const TPair<K,V>& B, ContextType& Ctx) const
+	{
+		return Key.DiffMember(A.Key, B.Key, MemberIds[0], Ctx) || 
+			Value.DiffMember(A.Value, B.Value, MemberIds[1], Ctx);
+	}
 
 	const FMemberId				MemberIds[2];
 	const TMemberSerializer<K>	Key;
@@ -516,7 +521,8 @@ struct FTransformBinding : ICustomBinding
 
 	PLAINPROPS_API void	Save(FMemberBuilder& Dst, const FTransform& Src, const FTransform* Default, const FSaveContext& Context) const;
 	PLAINPROPS_API void	Load(FTransform& Dst, FStructLoadView Src, ECustomLoadMethod Method) const;
-	inline static bool	Diff(const FTransform& A, const FTransform& B) { return !A.Equals(B, 0.0); }
+	PLAINPROPS_API bool Diff(const FTransform& A, const FTransform& B, FDiffContext& Ctx) const;
+	inline static  bool Diff(const FTransform& A, const FTransform& B, const FBindContext&) { return !A.Equals(B, 0.0); }
 };
 
 //////////////////////////////////////////////////////////////////////////
@@ -533,7 +539,7 @@ struct FGuidBinding : ICustomBinding
 
 	PLAINPROPS_API void	Save(FMemberBuilder& Dst, const FGuid& Src, const FGuid* Default, const FSaveContext&) const;
 	PLAINPROPS_API void	Load(FGuid& Dst, FStructLoadView Src, ECustomLoadMethod Method) const;
-	inline static bool	Diff(FGuid A, FGuid B) { return A != B; }
+	inline static  bool Diff(FGuid A, FGuid B, const FBindContext&) { return A != B; }
 };
 
 //////////////////////////////////////////////////////////////////////////
@@ -550,7 +556,7 @@ struct FColorBinding : ICustomBinding
 
 	PLAINPROPS_API void	Save(FMemberBuilder& Dst, const FColor& Src, const FColor* Default, const FSaveContext&) const;
 	PLAINPROPS_API void	Load(FColor& Dst, FStructLoadView Src, ECustomLoadMethod Method) const;
-	inline static bool	Diff(FColor A, FColor B) { return A != B; }
+	inline static  bool Diff(FColor A, FColor B, const FBindContext&) { return A != B; }
 };
 
 //////////////////////////////////////////////////////////////////////////
@@ -567,7 +573,7 @@ struct FLinearColorBinding : ICustomBinding
 
 	PLAINPROPS_API void	Save(FMemberBuilder& Dst, const FLinearColor& Src, const FLinearColor* Default, const FSaveContext&) const;
 	PLAINPROPS_API void	Load(FLinearColor& Dst, FStructLoadView Src, ECustomLoadMethod Method) const;
-	inline static bool	Diff(FLinearColor A, FLinearColor B) { return A != B; }
+	inline static  bool Diff(FLinearColor A, FLinearColor B, const FBindContext&) { return A != B; }
 };
 
 //////////////////////////////////////////////////////////////////////////
@@ -811,7 +817,7 @@ struct TBaseDeltaBinding : ICustomBinding, FBaseDeltaBinding
 		}
 	}
 
-	inline static bool Diff(const Type& As, const Type& Bs)
+	inline static bool Diff(const Type& As, const Type& Bs, const FBindContext&)
 	{
 		if (As.Num() != Bs.Num())
 		{
diff --git a/Engine/Plugins/Experimental/PlainProps/Source/Tests/PlainPropsUeCoreTests.cpp b/Engine/Plugins/Experimental/PlainProps/Source/Tests/PlainPropsUeCoreTests.cpp
--- a/Engine/Plugins/Experimental/PlainProps/Source/Tests/PlainPropsUeCoreTests.cpp
+++ b/Engine/Plugins/Experimental/PlainProps/Source/Tests/PlainPropsUeCoreTests.cpp
@@ -109,12 +109,7 @@ struct FNameDeclaration
 	}
 };
 
-struct FTestCustomBinding : ICustomBinding
-{
-//	virtual FBothStructId GetId() const = 0;
-};
-
-struct FNameBinding : FTestCustomBinding
+struct FNameBinding final : ICustomBinding
 {
 	virtual void SaveCustom(FMemberBuilder& Dst, const void* Src, const void*, const FSaveContext& Ctx) override
 	{
@@ -128,15 +123,13 @@ struct FNameBinding : FTestCustomBinding
 		*static_cast<FName*>(Dst) = Names.Get(Idx);
 	}
 
-	virtual bool DiffCustom(const void* StructA, const void* StructB) const override
+	virtual bool DiffCustom(const void* StructA, const void* StructB, const FBindContext&) const override
 	{
 		FName A = *static_cast<const FName*>(StructA);
 		FName B = *static_cast<const FName*>(StructB);
 		return A.IsEqual(B, ENameCase::CaseSensitive);
 	}
 
-//	virtual FBothStructId GetId() const override { return { Declaration.Id, Declaration.Id }; }
-
 	FNameDeclaration	Declaration;
 	TSet<FName>			Names;
 };
@@ -173,14 +166,14 @@ template<class T>
 void FBatchSaver::Save(T&& Object) 
 {
 	FBothStructId Id = IndexStructBothId<FRuntimeIds, TTypename<std::remove_reference_t<T>>>();
-	SavedObjects.Emplace(Id.DeclId, SaveStruct(&Object, Id.BindId, {GTypes, GSchemas, Customs, Scratch}));
+	SavedObjects.Emplace(Id.DeclId, SaveStruct(&Object, Id.BindId, {{GTypes, GSchemas, Customs}, Scratch}));
 }
 
 template<class T>
 bool FBatchSaver::SaveDelta(const T& Object, const T& Default) 
 {
 	FBothStructId Id = IndexStructBothId<FRuntimeIds, TTypename<std::remove_reference_t<T>>>();
-	if (const FBuiltStruct* Delta = SaveStructDeltaIfDiff(&Object, &Default, Id.BindId, {GTypes, GSchemas, Customs, Scratch}))
+	if (const FBuiltStruct* Delta = SaveStructDeltaIfDiff(&Object, &Default, Id.BindId, {{GTypes, GSchemas, Customs}, Scratch}))
 	{
 		SavedObjects.Emplace(Id.DeclId, MoveTemp(Delta));
 		return true;
@@ -655,7 +648,7 @@ struct FNDCBinding : ICustomBinding
 		}
 	}
 
-	static bool Diff(FNDC A, FNDC B)
+	static bool Diff(FNDC A, FNDC B, const FBindContext&)
 	{
 		return !(A == B);
 	}
diff --git a/Engine/Plugins/Experimental/PlainPropsUObject/Source/Private/PlainPropsUObjectBindings.cpp b/Engine/Plugins/Experimental/PlainPropsUObject/Source/Private/PlainPropsUObjectBindings.cpp
--- a/Engine/Plugins/Experimental/PlainPropsUObject/Source/Private/PlainPropsUObjectBindings.cpp
+++ b/Engine/Plugins/Experimental/PlainPropsUObject/Source/Private/PlainPropsUObjectBindings.cpp
@@ -568,7 +568,7 @@ void FFieldPathBinding::Load(FFieldPath& Dst, FStructLoadView Src, ECustomLoadMe
 	LoadStruct(&(Dst.*GFieldPathOwner), Members.GrabStruct());
 }
 
-bool FFieldPathBinding::Diff(const FFieldPath& A, const FFieldPath& B)
+bool FFieldPathBinding::Diff(const FFieldPath& A, const FFieldPath& B, const FBindContext&)
 {
 	return A != B;
 }
@@ -600,7 +600,7 @@ void FDelegateBinding::Load(FScriptDelegate& Dst, FStructLoadView Src, ECustomLo
 	}
 }
 
-bool FDelegateBinding::Diff(const FScriptDelegate& A, const FScriptDelegate& B)
+bool FDelegateBinding::Diff(const FScriptDelegate& A, const FScriptDelegate& B, const FBindContext&)
 {
 	return A != B;
 }
@@ -722,7 +722,7 @@ void FMulticastInlineDelegateBinding::Load(FMulticastScriptDelegate& Dst, FStruc
 	LoadMulticastDelegate(Dst, Members);
 }
 
-bool FMulticastInlineDelegateBinding::Diff(const FMulticastScriptDelegate& A, const FMulticastScriptDelegate& B)
+bool FMulticastInlineDelegateBinding::Diff(const FMulticastScriptDelegate& A, const FMulticastScriptDelegate& B, const FBindContext&)
 {
 	return DiffMulticastDelegate(A, B);
 }
@@ -741,7 +741,7 @@ struct FMulticastSparseDelegateBinding final : ICustomBinding
 
 	virtual void SaveCustom(FMemberBuilder& Dst, const void* Src, const void* Default, const FSaveContext& Ctx) override
 	{
-		if (!Default || DiffCustom(Src, Default))
+		if (!Default || DiffCustom(Src, Default, Ctx))
 		{
 			Save(Dst, *static_cast<const FSparseDelegate*>(Src), Ctx);
 		}
@@ -753,7 +753,7 @@ struct FMulticastSparseDelegateBinding final : ICustomBinding
 		Load(*static_cast<FSparseDelegate*>(Dst), Src);
 	}
 
-	virtual bool DiffCustom(const void* A, const void* B) const override
+	virtual bool DiffCustom(const void* A, const void* B, const FBindContext&) const override
 	{
 		return Diff(*static_cast<const FSparseDelegate*>(A), *static_cast<const FSparseDelegate*>(B));
 	}
@@ -840,7 +840,7 @@ void FVerseFunctionBinding::Load(FVerseFunction& Dst, FStructLoadView Src, ECust
 	unimplemented();
 }
 
-bool FVerseFunctionBinding::Diff(const FVerseFunction& A, const FVerseFunction& B)
+bool FVerseFunctionBinding::Diff(const FVerseFunction& A, const FVerseFunction& B, const FBindContext&)
 {
 	//return !(A == B);
 	return false;
@@ -858,7 +858,7 @@ void FDynamicallyTypedValueBinding::Load(Type& Dst, FStructLoadView Src, ECustom
 	unimplemented();
 }
 
-bool FDynamicallyTypedValueBinding::Diff(const ::UE::FDynamicallyTypedValue& A, const ::UE::FDynamicallyTypedValue& B)
+bool FDynamicallyTypedValueBinding::Diff(const ::UE::FDynamicallyTypedValue& A, const ::UE::FDynamicallyTypedValue& B, const FBindContext&)
 {
 	//return !UE::Verse::FRuntimeTypeDynamic::Get().AreIdentical(&A, &B);
 	return false;
@@ -876,7 +876,7 @@ void FReferencePropertyBinding::Load(FReferencePropertyValue& Dst, FStructLoadVi
 	unimplemented();
 }
 
-bool FReferencePropertyBinding::Diff(const FReferencePropertyValue& A, const FReferencePropertyValue& B)
+bool FReferencePropertyBinding::Diff(const FReferencePropertyValue& A, const FReferencePropertyValue& B, const FBindContext&)
 {
 	unimplemented();
 	return false;
@@ -892,7 +892,7 @@ struct FInterfaceBinding final : ICustomBinding
 
 	virtual void SaveCustom(FMemberBuilder& Dst, const void* Src, const void* Default, const FSaveContext& Ctx) override
 	{
-		if (!Default || DiffCustom(Src, Default))
+		if (!Default || DiffCustom(Src, Default, Ctx))
 		{
 			Save(Dst, *static_cast<const FScriptInterface*>(Src), Ctx);
 		}
@@ -904,7 +904,7 @@ struct FInterfaceBinding final : ICustomBinding
 		Load(*static_cast<FScriptInterface*>(Dst), Src);
 	}
 
-	virtual bool DiffCustom(const void* A, const void* B) const override
+	virtual bool DiffCustom(const void* A, const void* B, const FBindContext&) const override
 	{
 		return *static_cast<const FScriptInterface*>(A) != *static_cast<const FScriptInterface*>(B);
 	}
@@ -2077,7 +2077,7 @@ struct TSetPropertyBinding : IItemRangeBinding, ICustomBinding
 		LoadSetDelta(*this, *static_cast<FScriptSet*>(Dst), Src);
 	}
 
-	virtual bool DiffCustom(const void* A, const void* B) const override
+	virtual bool DiffCustom(const void* A, const void* B, const FBindContext&) const override
 	{
 		return DiffSet(*this, *static_cast<const FScriptSet*>(A), *static_cast<const FScriptSet*>(B));
 	}
@@ -2579,7 +2579,7 @@ struct TMapPropertyCustomBinding : TMapPropertyItemBinding<FScriptMap, KeyKind,
 		InnerKey.DestroyItem(Tmp);
 	}
 
-	virtual bool DiffCustom(const void* A, const void* B) const override
+	virtual bool DiffCustom(const void* A, const void* B, const FBindContext&) const override
 	{
 		return DiffSet(*this, *static_cast<const FScriptMap*>(A), *static_cast<const FScriptMap*>(B));
 	}
@@ -3529,7 +3529,7 @@ struct TCustomPropertyBinding final : ICustomBinding
 
 	virtual void SaveCustom(FMemberBuilder& Dst, const void* Src, const void* Default, const FSaveContext& Ctx) override
 	{
-		if (!Default || DiffCustom(Src, Default))
+		if (!Default || DiffCustom(Src, Default, Ctx))
 		{
 			Batch.Save(Dst, *static_cast<const Type*>(Src), Ctx);
 		}
@@ -3540,7 +3540,7 @@ struct TCustomPropertyBinding final : ICustomBinding
 		Batch.Load(*static_cast<Type*>(Dst), Src);
 	}
 
-	virtual bool DiffCustom(const void* A, const void* B) const override
+	virtual bool DiffCustom(const void* A, const void* B, const FBindContext&) const override
 	{
 		return DiffProperty(*static_cast<const Type*>(A), *static_cast<const Type*>(B));
 	}
@@ -3597,8 +3597,8 @@ class FBatchSaver
 {
 public:
 	FBatchSaver(FCustomBindings& Customs, int32 NumReserve)
-	: FlatCtx{GUE.Types, GUE.Schemas, Customs, Scratch, nullptr}
-	, DeltaCtx{GUE.Types, GUE.Schemas, Customs, Scratch, &GUE.Defaults}
+	: FlatCtx{{GUE.Types, GUE.Schemas, Customs}, Scratch, nullptr}
+	, DeltaCtx{{GUE.Types, GUE.Schemas, Customs}, Scratch, &GUE.Defaults}
 	{
 		SavedObjects.Reserve(NumReserve);
 	}
@@ -4223,11 +4223,11 @@ static int32 Roundtrip(ERoundtrip Options)
 
 		// Diff original vs PlainProps
 		UE_LOGFMT(LogPlainPropsUObject, Display, "Diffing UObjects roundtripped via PlainProps...");
-		FDiffContext DiffCtx = { GUE.Types, GUE.Schemas, Customs.Overlay };
+		FDiffContext DiffCtx = {{ GUE.Types, GUE.Schemas, Customs.Overlay }};
 		TUtf8StringBuilder<256> PpDiffs;
 		for (FInstance& Instance : Instances)
 		{
-			if (DiffStructs(Instance.Id, Instance.Orig, Instance.PP, /* in-out */ DiffCtx))
+			if (DiffStructs(Instance.Orig, Instance.PP, Instance.Id, /* in-out */ DiffCtx))
 			{
 				GPpDiff.Last = &Instance;
 				PrintDiff(/* out */ PpDiffs, GUE.Names, DiffCtx.Out);
@@ -4267,7 +4267,7 @@ static int32 Roundtrip(ERoundtrip Options)
 	{
 		// Diff original vs TPS
 		UE_LOGFMT(LogPlainPropsUObject, Display, "Diffing UObjects roundtripped via TPS...");
-		FDiffContext DiffCtx = { GUE.Types, GUE.Schemas, Customs.Overlay };
+		FDiffContext DiffCtx = {{ GUE.Types, GUE.Schemas, Customs.Overlay }};
 		TUtf8StringBuilder<256> TpsDiffs;
 		TArray<int32> TpsDiffIdxs;
 		for (FInstance& Instance : Instances)
@@ -4277,7 +4277,7 @@ static int32 Roundtrip(ERoundtrip Options)
 				continue;
 			}
 
-			if (DiffStructs(Instance.Id, Instance.Orig, Instance.TPS, /* in-out */ DiffCtx))
+			if (DiffStructs(Instance.Orig, Instance.TPS, Instance.Id, /* in-out */ DiffCtx))
 			{
 				GTpsDiff.Last = &Instance;
 				PrintDiff(/* out */ TpsDiffs, GUE.Names, DiffCtx.Out);
@@ -4322,7 +4322,7 @@ static int32 Roundtrip(ERoundtrip Options)
 
 		// Diff original vs UPS
 		UE_LOGFMT(LogPlainPropsUObject, Display, "Diffing UObjects roundtripped via UPS...");
-		FDiffContext DiffCtx = { GUE.Types, GUE.Schemas, Customs.Overlay };
+		FDiffContext DiffCtx = {{ GUE.Types, GUE.Schemas, Customs.Overlay }};
 		TUtf8StringBuilder<256> UpsDiffs;
 		TArray<int32> UpsDiffIdxs;
 		for (FInstance& Instance : Instances)
@@ -4332,7 +4332,7 @@ static int32 Roundtrip(ERoundtrip Options)
 				continue;
 			}
 
-			if (DiffStructs(Instance.Id, Instance.Orig, Instance.UPS, /* in-out */ DiffCtx))
+			if (DiffStructs(Instance.Orig, Instance.UPS, Instance.Id, /* in-out */ DiffCtx))
 			{
 				GUpsDiff.Last = &Instance;
 				PrintDiff(/* out */ UpsDiffs, GUE.Names, DiffCtx.Out);
diff --git a/Engine/Plugins/Experimental/PlainPropsUObject/Source/Public/PlainPropsUObjectRuntime.h b/Engine/Plugins/Experimental/PlainPropsUObject/Source/Public/PlainPropsUObjectRuntime.h
--- a/Engine/Plugins/Experimental/PlainPropsUObject/Source/Public/PlainPropsUObjectRuntime.h
+++ b/Engine/Plugins/Experimental/PlainPropsUObject/Source/Public/PlainPropsUObjectRuntime.h
@@ -191,7 +191,7 @@ struct FFieldPathBinding : ICustomBinding
 
 	void			Save(FMemberBuilder& Dst, const Type& Src, const Type* Default, const FSaveContext&) const;
 	void			Load(Type& Dst, FStructLoadView Src, ECustomLoadMethod Method) const;
-	static bool		Diff(const Type& A, const Type& B);
+	static bool		Diff(const Type& A, const Type& B, const FBindContext&);
 };
 
 struct FDelegateBinding : ICustomBinding
@@ -203,7 +203,7 @@ struct FDelegateBinding : ICustomBinding
 
 	void			Save(FMemberBuilder& Dst, const Type& Src, const Type* Default, const FSaveContext&) const;
 	void			Load(Type& Dst, FStructLoadView Src, ECustomLoadMethod Method) const;
-	static bool		Diff(const Type& A, const Type& B);
+	static bool		Diff(const Type& A, const Type& B, const FBindContext&);
 };
 
 struct FMulticastInlineDelegateBinding : ICustomBinding
@@ -215,7 +215,7 @@ struct FMulticastInlineDelegateBinding : ICustomBinding
 
 	void			Save(FMemberBuilder& Dst, const Type& Src, const Type* Default, const FSaveContext&) const;
 	void			Load(Type& Dst, FStructLoadView Src, ECustomLoadMethod Method) const;
-	static bool		Diff(const Type& A, const Type& B);
+	static bool		Diff(const Type& A, const Type& B, const FBindContext&);
 };
 
 //////////////////////////////////////////////////////////////////////////
@@ -229,7 +229,7 @@ struct FVerseFunctionBinding : ICustomBinding
 
 	void			Save(FMemberBuilder& Dst, const Type& Src, const Type* Default, const FSaveContext&) const;
 	void			Load(Type& Dst, FStructLoadView Src, ECustomLoadMethod Method) const;
-	static bool		Diff(const Type& A, const Type& B);
+	static bool		Diff(const Type& A, const Type& B, const FBindContext&);
 };
 
 struct FDynamicallyTypedValueBinding : ICustomBinding
@@ -241,7 +241,7 @@ struct FDynamicallyTypedValueBinding : ICustomBinding
 
 	void			Save(FMemberBuilder& Dst, const Type& Src, const Type* Default, const FSaveContext&) const;
 	void			Load(Type& Dst, FStructLoadView Src, ECustomLoadMethod Method) const;
-	static bool		Diff(const Type& A, const Type& B);
+	static bool		Diff(const Type& A, const Type& B, const FBindContext&);
 };
 
 
@@ -254,7 +254,7 @@ struct FReferencePropertyBinding : ICustomBinding
 
 	void			Save(FMemberBuilder& Dst, const Type& Src, const Type* Default, const FSaveContext&) const;
 	void			Load(Type& Dst, FStructLoadView Src, ECustomLoadMethod Method) const;
-	static bool		Diff(const Type& A, const Type& B);
+	static bool		Diff(const Type& A, const Type& B, const FBindContext&);
 };
 
 } // namespace PlainProps::UE
