PlainProps: Add "Text3DEditorFontSubsystem" to commandlet exclude list due to illegal value in enum EText3DFontStyleFlags

#rnx

diff --git a/Engine/Plugins/Experimental/PlainPropsUObject/Source/Private/PlainPropsUObjectBindings.cpp b/Engine/Plugins/Experimental/PlainPropsUObject/Source/Private/PlainPropsUObjectBindings.cpp
--- a/Engine/Plugins/Experimental/PlainPropsUObject/Source/Private/PlainPropsUObjectBindings.cpp
+++ b/Engine/Plugins/Experimental/PlainPropsUObject/Source/Private/PlainPropsUObjectBindings.cpp
@@ -3882,7 +3882,10 @@ static bool IncludeClass(UClass* Class)
 		"GameFeaturePluginStateMachine", // Cloning ensure
 		"WorldSettings", // QAGame
 		// CitySample - Enum value 8 is undeclared in /Script/Engine.ERichCurveTangentMode, illegal value detected in /Script/Engine.RichCurveKey::TangentMode
-		"AnimSequence", "AnimationSequencerDataModel", "MovieSceneControlRigParameterSection"
+		"AnimSequence", "AnimationSequencerDataModel", "MovieSceneControlRigParameterSection",
+		// Enum value 7 is undeclared in /Script/Text3D.EText3DFontStyleFlags, illegal value detected in /Script/Text3DEditor.Text3DEditorFont::FontStyleFlags
+		// FontStyleFlags	Monospace | Bold | Italic (7 '\a')	UnrealEditor-Text3DEditor.dll!EText3DFontStyleFlags
+		"Text3DEditorFontSubsystem",
 	};
 	static const FName SuperExclusions[] = { "LevelScriptActor" };
 
