PlainProps: Replace the FName indexer with an FSensitiveName indexer

#rnx

diff --git a/Engine/Plugins/Experimental/PlainProps/Source/Private/PlainPropsUeCoreBindings.cpp b/Engine/Plugins/Experimental/PlainProps/Source/Private/PlainPropsUeCoreBindings.cpp
--- a/Engine/Plugins/Experimental/PlainProps/Source/Private/PlainPropsUeCoreBindings.cpp
+++ b/Engine/Plugins/Experimental/PlainProps/Source/Private/PlainPropsUeCoreBindings.cpp
@@ -5,10 +5,33 @@
 #include "PlainPropsDiff.h"
 #include "PlainPropsSave.h"
 #include "Math/Transform.h"
+#include "Misc/DefinePrivateMemberPtr.h"
+
+#if UE_FNAME_OUTLINE_NUMBER
+#if WITH_CASE_PRESERVING_NAME
+UE_DEFINE_PRIVATE_MEMBER_PTR(FNameEntryId, GDisplayIndexInternal, FName, DisplayIndex);
+#endif
+UE_DEFINE_PRIVATE_MEMBER_PTR(FNameEntryId, GComparisonIndexInternal, FName, ComparisonIndex);
+#endif
 
 namespace PlainProps::UE
 {
 
+#if UE_FNAME_OUTLINE_NUMBER // dll exported FSensitiveName::InitXXX functions that rely on accessing private members of FName
+#if WITH_CASE_PRESERVING_NAME
+uint32 FSensitiveName::InitValue(FName Name)
+{
+	FNameEntryId DisplayId = Name.*GDisplayIndexInternal;
+	return DisplayId.ToUnstableInt() | (DisplayId != Name.*GComparisonIndexInternal) * DifferentIdsFlag;
+}
+#else
+FNameEntryId FSensitiveName::InitDisplayId(FName Name)
+{
+	return Name.*GComparisonIndexInternal;
+}
+#endif
+#endif
+
 FMemberId FNumeralGenerator::Grow(int32 Max)
 {
 	check(Max >= Cache.Num());
@@ -202,7 +225,7 @@ namespace PlainProps
 {
 	
 template<>
-void AppendString(FUtf8StringBuilderBase& Out, const FName& Name)
+void AppendString(FUtf8StringBuilderBase& Out, const UE::FSensitiveName& Name)
 {
 	Name.AppendString(Out);
 }
diff --git a/Engine/Plugins/Experimental/PlainProps/Source/Public/PlainPropsUeCoreBindings.h b/Engine/Plugins/Experimental/PlainProps/Source/Public/PlainPropsUeCoreBindings.h
--- a/Engine/Plugins/Experimental/PlainProps/Source/Public/PlainPropsUeCoreBindings.h
+++ b/Engine/Plugins/Experimental/PlainProps/Source/Public/PlainPropsUeCoreBindings.h
@@ -38,16 +38,108 @@ struct TIsContiguousContainer<PlainProps::TRangeView<T>>
 namespace PlainProps::UE
 {
 
+class FSensitiveName
+{
+public:
+#if UE_FNAME_OUTLINE_NUMBER
+	using IntType = uint32;
+#else
+	using IntType = uint64;
+#endif
+	inline explicit FSensitiveName(FName Name);
+	inline FName ToName() const { return FName(GetComparisonId(), GetDisplayId(), GetNumber()); }
+	inline void AppendString(FUtf8StringBuilderBase& Out) const { ToName().AppendString(Out); }
+	inline IntType ToUnstableInt() const;
+	inline static FSensitiveName FromUnstableInt(IntType UnstableInt) { return FSensitiveName(UnstableInt); }
+	[[nodiscard]] inline bool operator==(FSensitiveName Other) const { return ToUnstableInt() == Other.ToUnstableInt(); }
+	[[nodiscard]] inline bool operator!=(FSensitiveName Other) const { return !(*this == Other); }
+	[[nodiscard]] friend inline uint32 GetTypeHash(FSensitiveName Name) { return GetTypeHash(Name.ToUnstableInt()); }
+
+private:
+	inline explicit FSensitiveName(IntType);
+
+#if WITH_CASE_PRESERVING_NAME
+	static constexpr uint32 DifferentIdsFlag = 1u << 31;
+	static constexpr uint32 DisplayIdMask = ~DifferentIdsFlag;
+ 	uint32 Value; // DisplayId + DifferentIdsFlag
+	inline bool SameIds() const { return (Value & DifferentIdsFlag) == 0; }
+	inline FNameEntryId GetDisplayId() const { return FNameEntryId::FromUnstableInt(Value & DisplayIdMask); }
+	inline FNameEntryId GetComparisonId() const { return SameIds() ? GetDisplayId() : FName::GetComparisonIdFromDisplayId(GetDisplayId()); }
+	inline uint32 GetUnstableValue() const { return Value; }						// ToUnstableInt() helper
+#else
+	FNameEntryId DisplayId;
+	inline FNameEntryId GetDisplayId() const { return DisplayId; }					// ToName() helper
+	inline FNameEntryId GetComparisonId() const { return DisplayId; }				// ToName() helper
+	inline uint32 GetUnstableValue() const { return DisplayId.ToUnstableInt(); }	// ToUnstableInt() helper
+#endif
+
+#if UE_FNAME_OUTLINE_NUMBER
+	inline uint32 GetNumber() const { return NAME_NO_NUMBER_INTERNAL; }				// ToName() helper
+	PLAINPROPS_API static uint32 InitValue(FName Name);								// Constructor helper for WITH_CASE_PRESERVING_NAME
+	PLAINPROPS_API static FNameEntryId InitDisplayId(FName Name);					// Constructor helper for !WITH_CASE_PRESERVING_NAME
+#else
+	int32 Number;
+	inline uint32 GetNumber() const { return Number; }								// ToName() helper
+	inline static uint32 InitValue(FName Name);										// Constructor helper for WITH_CASE_PRESERVING_NAME
+	inline static FNameEntryId InitDisplayId(FName Name);							// Constructor helper for !WITH_CASE_PRESERVING_NAME
+#endif
+};
+
+FSensitiveName::FSensitiveName(FName Name)
+#if WITH_CASE_PRESERVING_NAME
+	: Value(InitValue(Name))
+#else
+	: DisplayId(InitDisplayId(Name))
+#endif
+#if !UE_FNAME_OUTLINE_NUMBER
+	, Number(Name.GetNumber())
+#endif
+{}
+
+FSensitiveName::FSensitiveName(IntType UnstableInt)
+#if WITH_CASE_PRESERVING_NAME
+	: Value(static_cast<uint32>(UnstableInt))
+#else
+	: DisplayId(FNameEntryId::FromUnstableInt(static_cast<uint32>(UnstableInt)))
+#endif
+#if !UE_FNAME_OUTLINE_NUMBER
+	, Number(static_cast<uint32>(UnstableInt >> 32))
+#endif
+{}
+
+FSensitiveName::IntType FSensitiveName::ToUnstableInt() const
+{
+#if UE_FNAME_OUTLINE_NUMBER
+	return GetUnstableValue();
+#else
+	return (uint64(Number) << 32) | GetUnstableValue();
+#endif
+}
+
+#if !UE_FNAME_OUTLINE_NUMBER // inline FSensitiveName::InitXXX functions that rely only on the public FName API
+#if WITH_CASE_PRESERVING_NAME
+uint32 FSensitiveName::InitValue(FName Name)
+{
+	return Name.GetDisplayIndex().ToUnstableInt() | (Name.GetDisplayIndex() != Name.GetComparisonIndex()) * DifferentIdsFlag;
+}
+#else
+FNameEntryId FSensitiveName::InitDisplayId(FName Name)
+{
+	return Name.GetComparisonIndex();
+}
+#endif
+#endif
+
 // Temporary implementation for numeral structs
 class FNumeralGenerator
 {
-	TIdIndexer<FName>&			Indexer;
+	TIdIndexer<FSensitiveName>&	Indexer;
 	TArray<FMemberId>			Cache; // 0, 1, 2, ..
 
 	PLAINPROPS_API FMemberId	Grow(int32 Max);
 
 public:
-	FNumeralGenerator(TIdIndexer<FName>& InIndexer) : Indexer(InIndexer) { }
+	FNumeralGenerator(TIdIndexer<FSensitiveName>& InIndexer) : Indexer(InIndexer) { }
 
 	inline FMemberId			Make(uint16 Numeral)
 	{
@@ -1108,7 +1200,7 @@ struct TTypename<TMap<K, V, SetAllocator, KeyFuncs>>
 };
 
 template<>
-PLAINPROPS_API void AppendString(FUtf8StringBuilderBase& Out, const FName& Name);
+PLAINPROPS_API void AppendString(FUtf8StringBuilderBase& Out, const UE::FSensitiveName& Name);
 
 template<typename T, typename Allocator>
 struct TRangeBind<TArray<T, Allocator>>
diff --git a/Engine/Plugins/Experimental/PlainProps/Source/Tests/PlainPropsUeCoreTests.cpp b/Engine/Plugins/Experimental/PlainProps/Source/Tests/PlainPropsUeCoreTests.cpp
--- a/Engine/Plugins/Experimental/PlainProps/Source/Tests/PlainPropsUeCoreTests.cpp
+++ b/Engine/Plugins/Experimental/PlainProps/Source/Tests/PlainPropsUeCoreTests.cpp
@@ -33,13 +33,13 @@ DEFINE_LOG_CATEGORY_STATIC(LogPlainPropsUeCoreTests, Log, All);
 namespace PlainProps::UE::Test
 {
 
-static TIdIndexer<FName>		GNames;
-static FDebugIds				GDebug(GNames);
-static FDeclarations			GTypes(GDebug);
-static FSchemaBindings			GSchemas(GDebug);
-static FCustomBindingsBottom	GCustoms(GDebug);
-static FCustomBindingsOverlay	GDeltaCustoms(GCustoms);
-static FNumeralGenerator		GNumeralGenerator(GNames);
+static TIdIndexer<FSensitiveName>	GNames;
+static FDebugIds					GDebug(GNames);
+static FDeclarations				GTypes(GDebug);
+static FSchemaBindings				GSchemas(GDebug);
+static FCustomBindingsBottom		GCustoms(GDebug);
+static FCustomBindingsOverlay		GDeltaCustoms(GCustoms);
+static FNumeralGenerator			GNumeralGenerator(GNames);
 
 struct FRuntimeIds
 {
@@ -218,7 +218,7 @@ TArray64<uint8> FBatchSaver::Write(TArray<FStructId>* OutMemoryIds) const
 
 	// Filter out declared but unused names and ids
 	FWriter Writer(GNames, BindIds, Schemas, Format);
-	TArray<FName> UsedNames;
+	TArray<FSensitiveName> UsedNames;
 	UsedNames.Reserve(Writer.GetUsedNames().Num());
 	for (FNameId Name : Writer.GetUsedNames())
 	{
@@ -228,7 +228,7 @@ TArray64<uint8> FBatchSaver::Write(TArray<FStructId>* OutMemoryIds) const
 	// Write ids. Just copying in-memory FNames, a stable format might use SaveNameBatch().
 	TArray64<uint8> Out;
 	WriteInt(Out, Magics[0]);
-	WriteNumAndArray(Out, TArrayView<const FName, int32>(UsedNames));
+	WriteNumAndArray(Out, TArrayView<const FSensitiveName, int32>(UsedNames));
 
 	// Write schemas
 	WriteInt(Out, Magics[1]);
@@ -264,9 +264,9 @@ TArray64<uint8> FBatchSaver::Write(TArray<FStructId>* OutMemoryIds) const
 
 class FStableNameBatchIds final : public FStableBatchIds
 {
-	TConstArrayView<FName> Names;
+	TConstArrayView<FSensitiveName> Names;
 public:
-	FStableNameBatchIds(FSchemaBatchId Batch, TConstArrayView<FName> InNames) : FStableBatchIds(Batch), Names(InNames) {}
+	FStableNameBatchIds(FSchemaBatchId Batch, TConstArrayView<FSensitiveName> InNames) : FStableBatchIds(Batch), Names(InNames) {}
 	using FStableBatchIds::AppendString;
 
 	virtual uint32				NumNames() const override							{ return static_cast<uint32>(Names.Num()); }
@@ -362,7 +362,7 @@ public:
 		// Read ids
 		FByteReader It(Data);
 		CHECK(It.Grab<uint32>() == Magics[0]);
-		Names = GrabNumAndArray<FName>(It);
+		Names = GrabNumAndArray<FSensitiveName>(It);
 		CHECK(Names.IsEmpty() != InRuntimeIds.IsEmpty());
 		
 		// Read schemas
@@ -483,19 +483,19 @@ public:
 		++LoadIdx;
 	}
 
-	const FBatchIds&			GetBatchIds() { return *LoadBatchIds; }
+	const FBatchIds&				GetBatchIds() { return *LoadBatchIds; }
 
 private:
-	const FSchemaBatch*			SavedSchemas;
-	const FSchemaBatch*			TranslatedSchemas = nullptr;
-	TConstArrayView<FName>		Names;
-	FNameBinding				NameBinding;
-	TUniquePtr<FBatchIds>		LoadBatchIds;
-	FCustomBindingsOverlay		Customs;
-	FLoadBatchPtr				Plans;
-	TArray<FStructView>			Objects;
-	int32						LoadIdx = 0;
-	ESchemaFormat				Format;
+	const FSchemaBatch*				SavedSchemas;
+	const FSchemaBatch*				TranslatedSchemas = nullptr;
+	TConstArrayView<FSensitiveName>	Names;
+	FNameBinding					NameBinding;
+	TUniquePtr<FBatchIds>			LoadBatchIds;
+	FCustomBindingsOverlay			Customs;
+	FLoadBatchPtr					Plans;
+	TArray<FStructView>				Objects;
+	int32							LoadIdx = 0;
+	ESchemaFormat					Format;
 };
 
 static void Run(void (*Save)(FBatchSaver&), void (*Load)(FBatchLoader&), const FCustomBindings& Customs = GCustoms)
@@ -608,6 +608,30 @@ struct FNames
 };
 PP_REFLECT_STRUCT(PlainProps::UE::Test, FNames, void, Name, Names);
 
+struct FSensitive
+{
+	int32 X;
+	int32 X_0;
+	int32 X_1;
+	int32 x;
+	int32 x_0;
+	int32 x_1;
+	bool operator==(const FSensitive&) const = default;
+};
+PP_REFLECT_STRUCT(PlainProps::UE::Test, FSensitive, void, X, X_0, X_1, x, x_0, x_1);
+
+struct FSensitiveReversed
+{
+	int32 x_1;
+	int32 x_0;
+	int32 x;
+	int32 X_1;
+	int32 X_0;
+	int32 X;
+	bool operator==(const FSensitiveReversed&) const = default;
+};
+PP_REFLECT_STRUCT(PlainProps::UE::Test, FSensitiveReversed, void, x_1, x_0, x, X_1, X_0, X);
+
 struct FStr
 {
 	FString S;
@@ -1062,13 +1086,31 @@ TEST_CASE_NAMED(FPlainPropsUeCoreTest, "System::Core::Serialization::PlainProps:
 		Run([](FBatchSaver& Batch)
 			{
 				Batch.Save(FNames{ FName("A"), {FName("Y"), FName("A")} });
+				Batch.Save(FNames{ FName("B"), {FName("B_0"), FName("B_1")} });
 			}, 
 			[](FBatchLoader& Batch)
 			{
 				CHECK(Batch.Load<FNames>() == FNames{ FName("A"), {FName("Y"), FName("A")}});
+				CHECK(Batch.Load<FNames>() == FNames{ FName("B"), {FName("B_0"), FName("B_1")}});
 			});
 	}
 
+	SECTION("FSensitiveName")
+	{
+		TScopedDefaultStructBinding<FSensitive> Sensitive;
+		TScopedDefaultStructBinding<FSensitiveReversed> Sensitive2;
+		Run([](FBatchSaver& Batch)
+			{
+				Batch.Save(FSensitive{1, 2, 3, 4, 5, 6});
+				Batch.Save(FSensitiveReversed{-1, -2, -3, -4, -5, -6});
+			}, 
+			[](FBatchLoader& Batch)
+			{
+				CHECK(Batch.Load<FSensitive>() == FSensitive{1, 2, 3, 4, 5, 6});
+				CHECK(Batch.Load<FSensitiveReversed>() == FSensitiveReversed{-1, -2, -3, -4, -5, -6});
+			});
+	}
+	
 	SECTION("FString")
 	{
 		TScopedDefaultStructBinding<FStr> Str;
diff --git a/Engine/Plugins/Experimental/PlainPropsUObject/Source/Private/PlainPropsUObjectBindings.cpp b/Engine/Plugins/Experimental/PlainPropsUObject/Source/Private/PlainPropsUObjectBindings.cpp
--- a/Engine/Plugins/Experimental/PlainPropsUObject/Source/Private/PlainPropsUObjectBindings.cpp
+++ b/Engine/Plugins/Experimental/PlainPropsUObject/Source/Private/PlainPropsUObjectBindings.cpp
@@ -256,12 +256,12 @@ const void* FDefaultStructs::Get(FBindId Id)
 
 ////////////////////////////////////////////////////////////////////////////////////////////////
 
-FCommonScopeIds::FCommonScopeIds(TIdIndexer<FName>& Names)
+FCommonScopeIds::FCommonScopeIds(TIdIndexer<FSensitiveName>& Names)
 : Core(						Names.MakeScope("/Script/Core"))
 , CoreUObject(				Names.MakeScope("/Script/CoreUObject"))
 {}
 
-FCommonTypenameIds::FCommonTypenameIds(TIdIndexer<FName>& Names)
+FCommonTypenameIds::FCommonTypenameIds(TIdIndexer<FSensitiveName>& Names)
 : Optional(					Names.NameType("Optional"))
 , Map(						Names.NameType("Map"))
 , Set(						Names.NameType("Set"))
@@ -279,7 +279,7 @@ FCommonTypenameIds::FCommonTypenameIds(TIdIndexer<FName>& Names)
 , VerseString(				Names.NameType("VerseString"))
 {}
 
-FCommonStructIds::FCommonStructIds(const FCommonScopeIds& Scopes, TIdIndexer<FName>& Names)
+FCommonStructIds::FCommonStructIds(const FCommonScopeIds& Scopes, TIdIndexer<FSensitiveName>& Names)
 : Name(						Names.IndexStruct({Scopes.Core, Names.MakeTypename("Name")}))
 , Text(						Names.IndexStruct({Scopes.Core, Names.MakeTypename("Text")}))
 , Guid(						Names.IndexStruct({Scopes.Core, Names.MakeTypename("Guid")}))
@@ -300,7 +300,7 @@ FCommonStructIds::FCommonStructIds(const FCommonScopeIds& Scopes, TIdIndexer<FNa
 , ReferencePropertyValue(	Names.IndexStruct({Scopes.CoreUObject, Names.MakeTypename("ReferencePropertyValue")}))
 {}
 
-FCommonMemberIds::FCommonMemberIds(TIdIndexer<FName>& Names)
+FCommonMemberIds::FCommonMemberIds(TIdIndexer<FSensitiveName>& Names)
 : Key(						Names.NameMember("Key"))
 , Value(					Names.NameMember("Value"))
 , Assign(					Names.NameMember("Assign"))
@@ -3402,38 +3402,18 @@ static void InitBatchedProperties()
 	GUE.Defaults.BindZeroes(GUE.Structs.LazyObjectPtr, sizeof(FLazyObjectPtr), alignof(FLazyObjectPtr));
 }
 
-#if UE_FNAME_OUTLINE_NUMBER
-static uint32 ToInt(FName In)
-{
-	return In.GetDisplayIndex().ToUnstableInt();
-}
-static FName FromInt(uint32 In)
-{
-	return FName::CreateFromDisplayId(FNameEntryId::FromUnstableInt(In), NAME_NO_NUMBER_INTERNAL);
-}
-#else
-static uint64 ToInt(FName In)
-{
-	return (uint64(In.GetNumber()) << 32) | In.GetDisplayIndex().ToUnstableInt();
-}
-static FName FromInt(uint64 In)
-{
-	return FName::CreateFromDisplayId(FNameEntryId::FromUnstableInt(static_cast<uint32>(In)), static_cast<uint32>(In >> 32));
-}
-#endif
-
 struct FMemoryPropertyBatch
 {
 	TArray<FText> Texts; // Tricky to serialize intrusively
 
 	static void Save(FMemberBuilder& Out, FName In, const FSaveContext&)
 	{
-		Out.Add(GUE.Members.Id, ToInt(In));
+		Out.Add(GUE.Members.Id, FSensitiveName(In).ToUnstableInt());
 	}
 
 	static void Load(FName& Out, FStructLoadView In)
 	{
-		Out = FromInt(LoadSole<uint64>(In));
+		Out = FSensitiveName::FromUnstableInt(LoadSole<FSensitiveName::IntType>(In)).ToName();
 	}
 
 	void Save(FMemberBuilder& Out, const FText& In, const FSaveContext&)
@@ -3619,7 +3599,7 @@ public:
 		// Write out FNames when using StableNames
 		if (!OutRuntimeIds)
 		{
-			TArray<FName> UsedNames;
+			TArray<FSensitiveName> UsedNames;
 			UsedNames.Reserve(Writer.GetUsedNames().Num());
 			for (FNameId Name : Writer.GetUsedNames())
 			{
@@ -3627,7 +3607,7 @@ public:
 			}
 
 			WriteInt(Out, Magics[0]);
-			WriteNumAndArray(Out, TArrayView<const FName, int32>(UsedNames));
+			WriteNumAndArray(Out, TArrayView<const FSensitiveName, int32>(UsedNames));
 		}
 
 		// Write schemas
@@ -3782,7 +3762,7 @@ public:
 		}
 		return WriteValue(Idx);
 	}
-	virtual FArchive& operator<<(FName& Value) override				{ return WriteValue(ToInt(Value));}
+	virtual FArchive& operator<<(FName& Value) override				{ return WriteValue(FSensitiveName(Value).ToUnstableInt());}
 	virtual FArchive& operator<<(UObject*& Value) override			{ return WriteValue(reinterpret_cast<uint64&>(Value)); }
 	virtual FArchive& operator<<(FObjectPtr& Value) override		{ return WriteValue(reinterpret_cast<uint64&>(Value)); }
 	virtual FArchive& operator<<(FWeakObjectPtr& Value) override	{ return WriteValue(reinterpret_cast<uint64&>(Value)); }
@@ -3822,7 +3802,12 @@ public:
 		return *this;
 	}
 
-	virtual FArchive& operator<<(FName& Value) override				{ Value = FromInt(ReadValue<decltype(ToInt(FName()))>()); return *this; }
+	virtual FArchive& operator<<(FName& Value) override
+	{
+		Value = FSensitiveName::FromUnstableInt(ReadValue<FSensitiveName::IntType>()).ToName();
+		return *this;
+	}
+
 	virtual FArchive& operator<<(UObject*& Value) override			{ reinterpret_cast<uint64&>(Value) = ReadValue<uint64>(); return *this; }
 	virtual FArchive& operator<<(FObjectPtr& Value) override		{ reinterpret_cast<uint64&>(Value) = ReadValue<uint64>(); return *this; }
 	virtual FArchive& operator<<(FWeakObjectPtr& Value) override	{ reinterpret_cast<uint64&>(Value) = ReadValue<uint64>(); return *this; }
@@ -3951,9 +3936,9 @@ FORCENOINLINE void LoadArchive(FPropertyReader& Archive, TConstArrayView<FInstan
 
 class FStableNameBatchIds final : public FStableBatchIds
 {
-	TConstArrayView<FName> Names;
+	TConstArrayView<FSensitiveName> Names;
 public:
-	FStableNameBatchIds(FSchemaBatchId Batch, TConstArrayView<FName> InNames) : FStableBatchIds(Batch), Names(InNames) {}
+	FStableNameBatchIds(FSchemaBatchId Batch, TConstArrayView<FSensitiveName> InNames) : FStableBatchIds(Batch), Names(InNames) {}
 	using FStableBatchIds::AppendString;
 
 	virtual uint32				NumNames() const override							{ return static_cast<uint32>(Names.Num()); }
@@ -4059,11 +4044,11 @@ public:
 		FByteReader It(Data);
 
 		// Read FNames when using Stable Names
-		TConstArrayView<FName> Names;
+		TConstArrayView<FSensitiveName> Names;
 		if (Format == ESchemaFormat::StableNames)
 		{
 			verify(It.Grab<uint32>() == Magics[0]);
-			Names = GrabNumAndArray<FName>(It);
+			Names = GrabNumAndArray<FSensitiveName>(It);
 		}
 		
 		// Read and mount schemas
diff --git a/Engine/Plugins/Experimental/PlainPropsUObject/Source/Public/PlainPropsUObjectRuntime.h b/Engine/Plugins/Experimental/PlainPropsUObject/Source/Public/PlainPropsUObjectRuntime.h
--- a/Engine/Plugins/Experimental/PlainPropsUObject/Source/Public/PlainPropsUObjectRuntime.h
+++ b/Engine/Plugins/Experimental/PlainPropsUObject/Source/Public/PlainPropsUObjectRuntime.h
@@ -44,7 +44,7 @@ private:
 
 struct FCommonScopeIds
 {
-	explicit FCommonScopeIds(TIdIndexer<FName>& Names);
+	explicit FCommonScopeIds(TIdIndexer<FSensitiveName>& Names);
 	
 	FScopeId					Core;
 	FScopeId					CoreUObject;
@@ -52,7 +52,7 @@ struct FCommonScopeIds
 
 struct FCommonTypenameIds
 {
-	explicit FCommonTypenameIds(TIdIndexer<FName>& Names);
+	explicit FCommonTypenameIds(TIdIndexer<FSensitiveName>& Names);
 	
 	FConcreteTypenameId			Optional;
 	FConcreteTypenameId			Map;
@@ -73,7 +73,7 @@ struct FCommonTypenameIds
 
 struct FCommonStructIds
 {
-	explicit FCommonStructIds(const FCommonScopeIds& Scopes, TIdIndexer<FName>& Names);
+	explicit FCommonStructIds(const FCommonScopeIds& Scopes, TIdIndexer<FSensitiveName>& Names);
 	
 	FDualStructId			Name;
 	FDualStructId			Text;
@@ -97,7 +97,7 @@ struct FCommonStructIds
 
 struct FCommonMemberIds
 {
-	explicit FCommonMemberIds(TIdIndexer<FName>& Names);
+	explicit FCommonMemberIds(TIdIndexer<FSensitiveName>& Names);
 	
 	FMemberId					Key;
 	FMemberId					Value;
@@ -116,7 +116,7 @@ struct FGlobals
 {
 	FGlobals();
 
-	TIdIndexer<FName>			Names;
+	TIdIndexer<FSensitiveName>	Names;
 	FDeclarations				Types;
 	FSchemaBindings				Schemas;
 	FCustomBindingsBottom		Customs;
diff --git a/Engine/Source/Programs/PlainPropsTests/PlainPropsTests.Target.cs b/Engine/Source/Programs/PlainPropsTests/PlainPropsTests.Target.cs
--- a/Engine/Source/Programs/PlainPropsTests/PlainPropsTests.Target.cs
+++ b/Engine/Source/Programs/PlainPropsTests/PlainPropsTests.Target.cs
@@ -12,6 +12,11 @@ public class PlainPropsTestsTarget : TestTargetRules
 		bCompileWithPluginSupport = true;
 		bCompileAgainstCoreUObject = false;
 
+		// Build tests with different permuations of FName case preserving and number outlining
+		// BuildEnvironment = TargetBuildEnvironment.Unique;
+		// bFNameOutlineNumber = true;
+		// GlobalDefinitions.Add("WITH_CASE_PRESERVING_NAME=0");
+
 		bBuildWithEditorOnlyData = Target.Platform.IsInGroup(UnrealPlatformGroup.Desktop)
 			&& (Target.Configuration == UnrealTargetConfiguration.Debug || Target.Configuration == UnrealTargetConfiguration.Development);
 	}
