PlainProps: Add natvis support for FSensitiveName

#rnx

diff --git a/Engine/Plugins/Experimental/PlainProps/Source/Tests/PlainPropsUeCoreTests.cpp b/Engine/Plugins/Experimental/PlainProps/Source/Tests/PlainPropsUeCoreTests.cpp
--- a/Engine/Plugins/Experimental/PlainProps/Source/Tests/PlainPropsUeCoreTests.cpp
+++ b/Engine/Plugins/Experimental/PlainProps/Source/Tests/PlainPropsUeCoreTests.cpp
@@ -969,7 +969,7 @@ TArray<T, TInlineAllocator<N>> MakeInlArray(const T* Str)
 
 TEST_CASE_NAMED(FPlainPropsUeCoreTest, "System::Core::Serialization::PlainProps::UE::Core", "[Core][PlainProps][SmokeFilter]")
 {
-	DbgVis::FIdScope _(GNames, "FName");
+	DbgVis::FIdScope _(GNames, "SensName");
 
 	SECTION("Basic")
 	{
diff --git a/Engine/Plugins/Experimental/PlainPropsUObject/Source/Private/PlainPropsUObjectBindings.cpp b/Engine/Plugins/Experimental/PlainPropsUObject/Source/Private/PlainPropsUObjectBindings.cpp
--- a/Engine/Plugins/Experimental/PlainPropsUObject/Source/Private/PlainPropsUObjectBindings.cpp
+++ b/Engine/Plugins/Experimental/PlainPropsUObject/Source/Private/PlainPropsUObjectBindings.cpp
@@ -4359,7 +4359,7 @@ int32 UTestPlainPropsCommandlet::Main(const FString& Params)
 {
 	using namespace PlainProps;
 	using namespace PlainProps::UE;
-	DbgVis::FIdScope _(GUE.Names, "FName");
+	DbgVis::FIdScope _(GUE.Names, "SensName");
 
 	if (int32 LoadIdx = Params.Find("-load="); LoadIdx != INDEX_NONE)
 	{
