PlainProps move to typed declarations in preparation of upgraders

#rnx

diff --git a/Engine/Plugins/Experimental/PlainProps/Source/Private/PlainPropsBind.cpp b/Engine/Plugins/Experimental/PlainProps/Source/Private/PlainPropsBind.cpp
--- a/Engine/Plugins/Experimental/PlainProps/Source/Private/PlainPropsBind.cpp
+++ b/Engine/Plugins/Experimental/PlainProps/Source/Private/PlainPropsBind.cpp
@@ -1,12 +1,14 @@
 // Copyright Epic Games, Inc. All Rights Reserved.
 
 #include "PlainPropsBind.h"
+#include "PlainPropsDeclare.h"
 #include "PlainPropsDiff.h"
 #include "PlainPropsIndex.h"
 #include "PlainPropsInternalBind.h"
 #include "PlainPropsInternalBuild.h"
 #include "PlainPropsInternalFormat.h"
 #include "PlainPropsInternalRead.h"
+#include "Templates/RefCounting.h"
 
 namespace PlainProps
 {
@@ -198,15 +200,33 @@ static void UpdateCulling(FIdWindow& Window, uint32 Id)
 
 struct FCustomBindingHandle
 {
-	ICustomBinding**				Bindings = nullptr;
+	ICustomBinding**				Value = nullptr;
 	uint32							Max = 0;
 
-	explicit						operator bool() { return !!Bindings;}
-	const FStructDeclaration**		GetDeclarations() { return reinterpret_cast<const FStructDeclaration**>((void*)(Bindings + Max)); }
-	FInnersHandle*					GetLoweredInners() { return reinterpret_cast<FInnersHandle*>(Bindings + 2 * Max); }
+	explicit						operator bool()		{ return !!Value; }
+	ICustomBinding*&				GetBinding()		{ return Value[0]; }
+	FInnersHandle&					GetInners()			{ return reinterpret_cast<FInnersHandle&>(Value[Max]); }
+	FStructDeclarationPtr&			GetDeclaration()	{ return reinterpret_cast<FStructDeclarationPtr&>(Value[Max*2]); }
+
+	void							Free()				{ GetDeclaration().~FStructDeclarationPtr(); }
 };
 
-void FCustomBindingMap::Bind(FBindId Id, ICustomBinding& Binding, const FStructDeclaration& Declaration, FInnersHandle LoweredInners)
+FCustomBindingMap::FCustomBindingMap(FDebugIds Dbg)
+: Debug(Dbg)
+{}
+
+FCustomBindingMap::~FCustomBindingMap()
+{
+	for (TSet<FBindId>::TConstIterator It(Keys); It; ++It)
+	{
+		int32 Idx = It.GetId().AsInteger();
+		FCustomBindingHandle Handle = {Values + Idx, MaxValues};
+		Handle.Free();
+	}
+	FMemory::Free(Values);
+}
+
+void FCustomBindingMap::Bind(FBindId Id, ICustomBinding& Binding, FStructDeclarationPtr&& Declaration, FInnersHandle LoweredInners)
 {
 	// Add key
 	bool bExists = false;
@@ -217,27 +237,28 @@ void FCustomBindingMap::Bind(FBindId Id, ICustomBinding& Binding, const FStructD
 	uint32 MaxKeys = static_cast<uint32>(Keys.GetMaxIndex());
 	if (MaxValues < MaxKeys)
 	{
-		FCustomBindingHandle Old = { Values, MaxValues };
-		SIZE_T Size = sizeof(ICustomBinding*) + sizeof(FStructDeclaration*) + sizeof(FInnersHandle);
+		constexpr SIZE_T ValueSize = sizeof(ICustomBinding*) + sizeof(FStructDeclarationPtr) + sizeof(FInnersHandle);
 
+		FCustomBindingHandle Old = { Values, MaxValues };
 		MaxValues = FMath::RoundUpToPowerOfTwo(FMath::Max(MaxKeys, 4u));
-		Values = reinterpret_cast<ICustomBinding**>(FMemory::MallocZeroed(MaxValues * Size));
+		Values = static_cast<ICustomBinding**>(FMemory::MallocZeroed(MaxValues * ValueSize));
 
 		if (Old)
 		{
 			FCustomBindingHandle New = { Values, MaxValues };
-			FMemory::Memcpy(New.Bindings,			Old.Bindings,			Old.Max * sizeof(ICustomBinding*));
-			FMemory::Memcpy(New.GetDeclarations(),	Old.GetDeclarations(),	Old.Max * sizeof(FStructDeclaration*));
-			FMemory::Memcpy(New.GetLoweredInners(), Old.GetLoweredInners(), Old.Max * sizeof(FInnersHandle));
-			FMemory::Free(Old.Bindings);
+			FMemory::Memcpy(&New.GetBinding(),		&Old.GetBinding(),		Old.Max * sizeof(ICustomBinding*));
+			FMemory::Memcpy(&New.GetDeclaration(),	&Old.GetDeclaration(),	Old.Max * sizeof(FStructDeclaration*));
+			FMemory::Memcpy(&New.GetInners(),		&Old.GetInners(),		Old.Max * sizeof(FInnersHandle));
+			FMemory::Free(Old.Value);
 		}
 	}
 
 	// Add value
-	FCustomBindingHandle Value(Values + SetIdx, MaxValues);
-	*Value.Bindings = &Binding;
-	*Value.GetDeclarations() = &Declaration;
-	*Value.GetLoweredInners() = LoweredInners;
+	FCustomBindingHandle Value = {Values + SetIdx, MaxValues};
+	Value.GetBinding() = &Binding;
+	Value.GetInners() = LoweredInners;
+	new (&Value.GetDeclaration()) FStructDeclarationPtr(MoveTemp(Declaration));
+	check(LoweredInners.Idx == Value.GetInners().Idx);
 
 	UpdateCulling(/* in-out */ Window, Id.Idx);
 }
@@ -256,6 +277,12 @@ inline FCustomBindingHandle FCustomBindingMap::Find(FBindId Id) const
 
 void FCustomBindingMap::Drop(FBindId Id)
 {
+	if (FSetElementId Idx = Keys.FindId(Id); Idx.IsValidId())
+	{
+		FCustomBindingHandle Handle = {Values + Idx.AsInteger(), MaxValues};
+		Handle.Free();
+	}	
+
 	int32 NumRemoved = Keys.Remove(Id);
 	checkf(NumRemoved == 1, TEXT("'%s' unbound"), *Debug.Print(Id));
 }
@@ -289,40 +316,54 @@ static bool operator==(FInnerStruct A, FInnerStruct B) { return A.Name == B.Name
 
 ////////////////////////////////////////////////////////////////////////////////////////////////
 
-void FCustomBindings::BindStruct(FBindId Id, ICustomBinding& Binding, const FStructDeclaration& Declaration, TConstArrayView<FInnerStruct> LoweredInners)
+FCustomBindings::FCustomBindings(TArray<FInnerStruct>& Inners, FDebugIds Dbg) : Map(Dbg), BottomInners(Inners) {}
+FCustomBindings::FCustomBindings(const FCustomBindings* Under) : Map(Under->Map.Debug), BottomInners(Under->BottomInners) {}
+FCustomBindings::~FCustomBindings() {}
+
+void FCustomBindings::BindStruct(FBindId Id, ICustomBinding& Binding, FStructDeclarationPtr&& Declaration, TConstArrayView<FInnerStruct> LoweredInners)
 {
 	FInnersHandle StoredInners = StoreInners(/* out */ BottomInners, LoweredInners);
 	check(Algo::Compare(LoweredInners, FetchInners(BottomInners, StoredInners)));
-	Map.Bind(Id, Binding, Declaration, StoredInners);
+	Map.Bind(Id, Binding, MoveTemp(Declaration), StoredInners);
+}
+
+void FCustomBindings::BindStruct(FBindId Id, ICustomBinding& Binding, const FStructSpec& Spec, TConstArrayView<FInnerStruct> LoweredInners)
+{
+	return BindStruct(Id, Binding, Declare(Spec), LoweredInners);
 }
 
 const ICustomBinding* FCustomBindings::FindStruct(FBindId Id) const
 {
 	FCustomBindingHandle Handle = Find(Id);
-	return Handle ? *Handle.Bindings : nullptr;
+	return Handle ? Handle.GetBinding() : nullptr;
 }
 
 const ICustomBinding* FCustomBindings::FindStruct(FBindId Id, TConstArrayView<FInnerStruct>& OutLoweredInners) const
 {
 	if (FCustomBindingHandle Handle = Find(Id))
 	{
-		OutLoweredInners = FetchInners(BottomInners, *Handle.GetLoweredInners());
-		return *Handle.Bindings;
+		OutLoweredInners = FetchInners(BottomInners, Handle.GetInners());
+		return Handle.GetBinding();
 	}
 	return nullptr;
 }
 
 ICustomBinding* FCustomBindings::FindStructToSave(FBindId Id, const FStructDeclaration*& OutDeclaration) const
 {
-	FCustomBindingHandle Handle = Find(Id);
-	OutDeclaration = Handle ? *Handle.GetDeclarations() : nullptr;
-	return Handle ? *Handle.Bindings : nullptr;
+	if (FCustomBindingHandle Handle = Find(Id))
+	{
+		OutDeclaration = Handle.GetDeclaration();
+		return Handle.GetBinding();
+	}
+
+	OutDeclaration = nullptr;
+	return nullptr;
 }
 
 const FStructDeclaration* FCustomBindings::FindDeclaration(FBindId Id) const
-{ 
+{
 	FCustomBindingHandle Handle = Find(Id);
-	return Handle ? *Handle.GetDeclarations() : nullptr;
+	return Handle ? Handle.GetDeclaration() : nullptr;
 }
 
 void FCustomBindings::DropStruct(FBindId Id)
@@ -374,8 +415,8 @@ FBothType FMemberBinding::IndexParameterName(FIdIndexerBase& Ids) const
 	return {BindType, DeclType};
 }
 
-FSchemaBindings::~FSchemaBindings()
-{}
+FSchemaBindings::FSchemaBindings(FDebugIds In) : Debug(In) {}
+FSchemaBindings::~FSchemaBindings() {}
 
 static uint16 CountInnerSchemas(TConstArrayView<FMemberBinding> Members)
 {
@@ -417,14 +458,72 @@ struct FMemberBinder : FMemberBinderBase
 	FInnerId* InnerSchemaIt;
 };
 
-void FSchemaBindings::BindStruct(FBindId BindId, FDeclId DeclId, TConstArrayView<FMemberBinding> Members)
+static void IncRef(FStructDeclarationPtr& Ptr)
+{
+	Ptr->AddRef();
+}
+
+static void DecRef(FStructDeclarationPtr& Ptr)
+{
+	if (Ptr->Release())
+	{
+		// Hack, recreate new TRefCountPtr to null out internal pointer in lieu of a Release function
+		new (&Ptr) FStructDeclarationPtr();
+	}
+}
+
+static void DeclareOrValidate(FStructDeclarationPtr& Out, FStructSpec Info, TConstArrayView<FInnerId> InnerIds, const FSchemaBinding& Binding)
+{
+	const uint16 Num = IntCastChecked<uint16>(Info.MemberNames.Num());
+
+	check(Info.Id == Binding.DeclId);
+	check(Num == Binding.NumMembers);
+	check(!!Info.Super == Binding.HasSuper());
+
+	if (Out)
+	{
+		check(Out->Id == Info.Id);
+		check(Out->Super == Info.Super);
+		check(Out->Version == Info.Version);
+		check(Out->Occupancy == Info.Occupancy);
+		check(Algo::Compare(Out->GetMemberOrder(), Info.MemberNames));
+		check(Out->NumInnerIds == Binding.NumInnerSchemas);
+		check(Out->NumInnerRanges == Binding.NumInnerRanges);
+		// Todo: check Out->GetTypes() matches Binding.Members
+		IncRef(Out);
+		return;
+	}
+
+	// Make header and allocate
+	FStructDeclaration Header = { 1, Info.Id, Info.Super, Info.Version, Num, Binding.NumInnerSchemas, Binding.NumInnerRanges, Info.Occupancy };
+	FStructDeclaration* Ptr = new (FMemory::MallocZeroed(Header.CalculateSize())) FStructDeclaration;
+	FMemory::Memcpy(Ptr, &Header, sizeof(Header));
+
+	// Copy footer
+	FMemory::Memcpy(Ptr->MemberNames, Info.MemberNames.GetData(), Info.MemberNames.NumBytes());
+	FMemory::Memcpy(const_cast<FInnerId*>(Ptr->GetInnerIds()), InnerIds.GetData(), InnerIds.NumBytes());
+	FMemberType* TypeIt = const_cast<FMemberType*>(Ptr->GetTypes());
+	for (FMemberBindType BindType : MakeArrayView(Binding.Members, Binding.NumMembers))
+	{
+		*TypeIt++ = ToMemberType(BindType);
+	}
+	for (FRangeBinding Range : MakeArrayView(Binding.GetRangeBindings(), Binding.NumInnerRanges))
+	{
+		*TypeIt++ = FMemberType(Range.GetSizeType());
+	}
+
+	// Make smart pointer
+	new (&Out) FStructDeclarationPtr(Ptr, /* add ref */ false);
+}
+
+FSchemaBinding* FSchemaBindings::Bind(FBindId Id, TConstArrayView<FMemberBinding> Members, FDeclId DeclId) 
 {
 	// Make header, allocate and copy header
 	FSchemaBinding Header = { DeclId, IntCastChecked<uint16>(Members.Num()), CountInnerSchemas(Members), CountRanges(Members) };
-	FSchemaBinding* Schema = new (FMemory::MallocZeroed(Header.CalculateSize())) FSchemaBinding {Header};
+	FSchemaBinding* Out = new (FMemory::MallocZeroed(Header.CalculateSize())) FSchemaBinding {Header};
 
 	// Write footer
-	FMemberBinder Footer(*Schema);
+	FMemberBinder Footer(*Out);
 	for (const FMemberBinding& Member : Members)
 	{
 		TConstArrayView<FRangeBinding> Ranges = Member.RangeBindings;
@@ -444,12 +543,49 @@ void FSchemaBindings::BindStruct(FBindId BindId, FDeclId DeclId, TConstArrayView
 	}
 
 	// Bind
-	if (BindId.Idx >= static_cast<uint32>(Bindings.Num()))
+	if (Id.Idx >= static_cast<uint32>(Bindings.Num()))
 	{
-		Bindings.SetNum(BindId.Idx + 1);
+		Bindings.SetNum(Id.Idx + 1);
 	}
-	checkf(!Bindings[BindId.Idx], TEXT("'%s' already bound"), *Debug.Print(BindId));
-	Bindings[BindId.Idx].Reset(Schema);
+	checkf(!Bindings[Id.Idx], TEXT("'%s' already bound"), *Debug.Print(Id));
+	Bindings[Id.Idx].Reset(Out);
+
+	return Out;
+}
+
+FStructDeclarationPtr& FSchemaBindings::At(FDeclId Id)
+{
+	if (Id.Idx >= static_cast<uint32>(Declarations.Num()))
+	{
+		Declarations.SetNum(Id.Idx + 1);
+	}
+	return Declarations[Id.Idx];
+}
+
+void FSchemaBindings::BindStruct(FBindId Id, TConstArrayView<FMemberBinding> Members, FStructSpec Spec)
+{
+	FSchemaBinding* Schema = Bind(Id, Members, Spec.Id);
+
+	// Extract declared struct ids, Binding::GetInnerSchemas() can contain unlowered FBindIds
+	TArray<FInnerId, TInlineAllocator<16>> InnerIds;
+	for (FMemberSpecView Member : Spec.MemberTypes)
+	{
+		if (Member.InnermostId)
+		{
+			InnerIds.Emplace(Member.InnermostId.Get());
+		}
+	}
+	check(InnerIds.Num() == Schema->NumInnerSchemas);
+	DeclareOrValidate(/* in-out */ At(Spec.Id), Spec, InnerIds, *Schema);
+}
+
+void FSchemaBindings::BindStruct(FBindId Id, TConstArrayView<FMemberBinding> Members, FStructDeclarationPtr&& InDecl)
+{
+	FSchemaBinding* Schema = Bind(Id, Members, InDecl->Id);
+	
+	FStructDeclarationPtr& OutDecl = At(InDecl->Id);
+	check(OutDecl == nullptr || OutDecl == InDecl); // Require exact pointer equality for now
+	OutDecl = InDecl;
 }
 
 const FSchemaBinding* FSchemaBindings::FindStruct(FBindId Id) const
@@ -457,16 +593,50 @@ const FSchemaBinding* FSchemaBindings::FindStruct(FBindId Id) const
 	return Id.Idx < (uint32)Bindings.Num() ? Bindings[Id.Idx].Get() : nullptr;
 }
 
+const FSchemaBinding* FSchemaBindings::FindStruct(FBindId Id, const FStructDeclaration*& OutDeclaration) const
+{
+	if (const FSchemaBinding* Binding = FindStruct(Id))
+	{
+		OutDeclaration = Declarations[Binding->DeclId.Idx];
+		return Binding;
+	}
+	OutDeclaration = nullptr;
+	return nullptr;
+}
+
 const FSchemaBinding& FSchemaBindings::GetStruct(FBindId Id) const
 {
 	checkf(Id.Idx < (uint32)Bindings.Num() && Bindings[Id.Idx], TEXT("'%s' is unbound"), *Debug.Print(Id));
 	return *Bindings[Id.Idx].Get();
 }
 
+const FSchemaBinding& FSchemaBindings::GetStruct(FBindId Id, const FStructDeclaration*& OutDecl) const
+{
+	const FSchemaBinding& Out = GetStruct(Id);
+	OutDecl = Declarations[Out.DeclId.Idx];
+	return Out;
+}
+
+const FStructDeclaration& FSchemaBindings::GetDeclaration(FBindId Id) const
+{
+	return *Declarations[GetStruct(Id).DeclId.Idx];
+}
+
+const FStructDeclaration* FSchemaBindings::FindDeclaration(FBindId Id) const
+{
+	if (const FSchemaBinding* Binding = FindStruct(Id))
+	{
+		return Declarations[Binding->DeclId.Idx];
+	}
+	return nullptr;
+}
+
 void FSchemaBindings::DropStruct(FBindId Id)
 {
 	checkf(Id.Idx < (uint32)Bindings.Num() && Bindings[Id.Idx], TEXT("'%s' is unbound"), *Debug.Print(Id));
-	Bindings[Id.Idx].Reset();
+	TUniquePtr<FSchemaBinding>& Binding = Bindings[Id.Idx];
+	DecRef(Declarations[Binding->DeclId.Idx]);
+	Binding.Reset();
 }
 
 FDeclId FSchemaBindings::Lower(FBindId Id) const
@@ -476,10 +646,29 @@ FDeclId FSchemaBindings::Lower(FBindId Id) const
 
 //////////////////////////////////////////////////////////////////////////
 
-FDeclId FStructBindIds::Lower(FBindId Id) const
+const FStructDeclaration& FBindContext::GetDeclaration(FBindId Id) const
+{
+	FCustomBindingHandle Handle = Customs.Find(Id);
+	return Handle ? *Handle.GetDeclaration() : Schemas.GetDeclaration(Id);
+}
+
+//////////////////////////////////////////////////////////////////////////
+
+const FEnumDeclaration* FBindDeclarations::Find(FEnumId Id) const
+{
+	return &Enums.Get(Id);
+}
+
+const FStructDeclaration* FBindDeclarations::Find(FStructId Id) const
+{
+	FCustomBindingHandle Handle = Customs.Find(FBindId(Id));
+	return Handle ? &*Handle.GetDeclaration() : Schemas.FindDeclaration(FBindId(Id));
+}
+
+FDeclId FBindDeclarations::Lower(FBindId Id) const
 {
-	const FStructDeclaration* CustomDecl = Customs.FindDeclaration(Id);
-	return CustomDecl ? CustomDecl->Id : Schemas.Lower(Id);
+	FCustomBindingHandle Handle = Customs.Find(Id);
+	return Handle ? Handle.GetDeclaration()->Id : Schemas.Lower(Id);
 }
 
 //////////////////////////////////////////////////////////////////////////
@@ -494,7 +683,7 @@ TArray<FStructId> IndexRuntimeIds(const FSchemaBatch& Schemas, FIdIndexerBase& I
 	for (int32 Idx = 0; Idx < Out.Num(); ++Idx)
 	{
 		const FStructSchema& Schema = *reinterpret_cast<const FStructSchema*>(Base + Offsets[Idx]);
-		Out[Idx] = Indexer.IndexStruct(Schema.Type);
+		Out[Idx] = Indexer.IndexDeclId(Schema.Type);
 	}
 
 	return Out;
diff --git a/Engine/Plugins/Experimental/PlainProps/Source/Private/PlainPropsBuild.cpp b/Engine/Plugins/Experimental/PlainProps/Source/Private/PlainPropsBuild.cpp
--- a/Engine/Plugins/Experimental/PlainProps/Source/Private/PlainPropsBuild.cpp
+++ b/Engine/Plugins/Experimental/PlainProps/Source/Private/PlainPropsBuild.cpp
@@ -161,7 +161,7 @@ FBuiltStruct* FMemberBuilder::BuildAndReset(FScratchAllocator& Scratch, const FS
 			FBuiltMember& Member = *It;
 			for (; OrderIdx < Order.Num() && Order[OrderIdx] != Member.Name; ++OrderIdx)
 			{}	
-			checkf(OrderIdx < Order.Num(), TEXT("Member '%s' in '%s' %s"), *Debug.Print(Member.Name), *Debug.Print(Declared.Type),
+			checkf(OrderIdx < Order.Num(), TEXT("Member '%s' in '%s' %s"), *Debug.Print(Member.Name), *Debug.Print(Declared.Id),
 					Order.Contains(Member.Name) ? TEXT("appeared in non-declared order") : TEXT("is undeclared"));
 			++OrderIdx;
 		}
@@ -191,7 +191,7 @@ FBuiltStruct* FDenseMemberBuilder::BuildHomo(const FStructDeclaration& Declarati
 	FBuiltStruct* Out = reinterpret_cast<FBuiltStruct*>(Scratch.Allocate(NumBytes, alignof(FBuiltStruct)));
 	Out->NumMembers = static_cast<uint16>(Num);
 
-	const FMemberId* Names = Declaration.MemberOrder; 
+	const FMemberId* Names = Declaration.MemberNames; 
 	for (uint32 Idx = 0; Idx < Num; ++Idx)
 	{
 		new (Out->Members + Idx) FBuiltMember(Names[Idx], Schema, Values[Idx]);
diff --git a/Engine/Plugins/Experimental/PlainProps/Source/Private/PlainPropsBuildSchema.cpp b/Engine/Plugins/Experimental/PlainProps/Source/Private/PlainPropsBuildSchema.cpp
--- a/Engine/Plugins/Experimental/PlainProps/Source/Private/PlainPropsBuildSchema.cpp
+++ b/Engine/Plugins/Experimental/PlainProps/Source/Private/PlainPropsBuildSchema.cpp
@@ -74,54 +74,66 @@ struct FEnumSchemaBuilder
 
 //////////////////////////////////////////////////////////////////////////
 
-FSchemasBuilder::FSchemasBuilder(const FDeclarations& Types, const FIds& InIds, const IBindIds& InBindIds, FScratchAllocator& InScratch, ESchemaFormat InFormat)
-: FSchemasBuilder(Types.GetStructs(), Types.GetEnums(), InIds, InBindIds, InScratch, InFormat)
-{}
-
-FSchemasBuilder::FSchemasBuilder(FStructDeclarations InStructs, FEnumDeclarations InEnums, const FIds& InIds, const IBindIds& InBindIds, FScratchAllocator& InScratch, ESchemaFormat InFormat)
-: DeclaredStructs(InStructs)
-, DeclaredEnums(InEnums)
-, Ids(InIds)
-, BindIds(InBindIds)
+FSchemasBuilder::FSchemasBuilder(const FIds& Names, const IDeclarations& Types, FScratchAllocator& InScratch, ESchemaFormat InFormat)
+: Declarations(Types)
+, Ids(Names)
 , Format(InFormat)
 , Scratch(InScratch)
-, Debug(InIds)
+, Debug(Names)
 {}
 
 FSchemasBuilder::~FSchemasBuilder() {}
 
-template<class IdType, class BuilderType, typename ...Ts>
-FORCEINLINE BuilderType& GetOrEmplace(TSet<IdType>& Indexer, IdType Id, TPagedArray<BuilderType, 4096>& Builders, Ts&&... EmplaceArgs)
+FEnumSchemaBuilder&	FSchemasBuilder::NoteEnum(FEnumId Id)
 {
-	FSetElementId Idx = Indexer.FindId(Id);
+	checkf(!bBuilt, TEXT("Noted new members after building"));
+
+	FSetElementId Idx = EnumIndices.FindId(Id);
 	if (Idx.IsValidId())
 	{
-		return Builders[Idx.AsInteger()];
+		return Enums[Idx.AsInteger()];
 	}
 
-	Idx = Indexer.Emplace(Id);
-	check(Idx.AsInteger() == Builders.Num());
-	Builders.Emplace(Forward<Ts>(EmplaceArgs)...);
-	return Builders.Last();
-}
-	
-FORCEINLINE FEnumSchemaBuilder&	FSchemasBuilder::NoteEnum(FEnumId Id)
-{
-	checkf(!bBuilt, TEXT("Noted new members after building"));
-	checkf(Id.Idx < uint32(DeclaredEnums.Num()) && DeclaredEnums[Id.Idx], TEXT("Undeclared enum '%s' noted"), *Debug.Print(Id));
-	return GetOrEmplace(EnumIndices, Id, Enums, *DeclaredEnums[Id.Idx], Debug, Id);
+	auto Declaration = Declarations.Find(Id);
+	checkf(Declaration, TEXT("Undeclared enum '%s' noted"), *Debug.Print(Id));
+
+	Idx = EnumIndices.Emplace(Id);
+	check(Idx.AsInteger() == Enums.Num());
+	Enums.Emplace(*Declaration, Debug, Id);
+
+	return Enums.Last();
 }
 
-FORCEINLINE FStructSchemaBuilder& FSchemasBuilder::NoteStruct(FStructId Id)
+FStructSchemaBuilder& FSchemasBuilder::NoteStruct(FStructId Id)
 {
 	checkf(!bBuilt, TEXT("Noted new members after building"));
 
-	FDeclId DeclId	= Id.Idx < uint32(DeclaredStructs.Num()) && DeclaredStructs[Id.Idx] 
-					? FDeclId(Id) : BindIds.Lower(FBindId(Id));
+	// Id is either FBindId or FDeclId depending on Format
+	// Most bind ids are the same as decl ids though
+	if (FSetElementId Idx = StructIndices.FindId(Id); Idx.IsValidId())
+	{
+		return Structs[Idx.AsInteger()];
+	}
+	
+	const FStructDeclaration* Declaration = Declarations.Find(Id);
+	checkf(Declaration, TEXT("Undeclared struct '%s' noted"), *Debug.Print(Id));
+	
+	// StableNames format can have lowered ids, so builder might exist already
+	FStructId NoteId = (Format == ESchemaFormat::StableNames) ? FStructId(Declaration->Id) : Id;
+	if (NoteId != Id)
+	{
+		if (FSetElementId Idx = StructIndices.FindId(NoteId); Idx.IsValidId())
+		{
+			return Structs[Idx.AsInteger()];
+		}
+	}
+
+	// First time noted, make new builder
+	FSetElementId Idx = StructIndices.Emplace(NoteId);
+	check(Idx.AsInteger() == Structs.Num());
+	Structs.Emplace(NoteId, *Declaration, *this);
 
-	checkf(DeclId.Idx < uint32(DeclaredStructs.Num()) && DeclaredStructs[DeclId.Idx], TEXT("Undeclared struct '%s' noted"), *Debug.Print(DeclId));
-	FStructId NoteId = (Format == ESchemaFormat::StableNames) ? DeclId : Id;
-	return GetOrEmplace(StructIndices, NoteId, Structs, NoteId, *DeclaredStructs[DeclId.Idx], *this);
+	return Structs.Last();
 }
 
 void FSchemasBuilder::NoteStructAndMembers(FStructId Id, const FBuiltStruct& Struct)
@@ -156,10 +168,21 @@ void FSchemasBuilder::NoteInheritanceChains()
 {
 	for (int Idx = 0, Num = Structs.Num(); Idx < Num; ++Idx)
 	{
-		for (FOptionalDeclId Super = Structs[Idx].Declaration.Super; Super; Super = DeclaredStructs[Super.Get().Idx]->Super)
+		const FStructDeclaration* Declaration = &Structs[Idx].Declaration;
+		while (Declaration->Super)
 		{
-			FStructId Id = Super.Get();
-			GetOrEmplace(StructIndices, Id, Structs, Id, *DeclaredStructs[Id.Idx], *this);
+			FDeclId Super = Declaration->Super.Get();
+			if (StructIndices.Contains(Super))
+			{
+				break;
+			}
+
+			Declaration = Declarations.Find(Super);
+			checkf(Declaration, TEXT("Undeclared super struct '%s' noted"), *Debug.Print(Super));
+
+			FSetElementId ElemId = StructIndices.Emplace(Super);
+			check(ElemId.AsInteger() == Structs.Num());
+			Structs.Emplace(Super, *Declaration, *this);
 		}
 	}
 }
@@ -252,7 +275,7 @@ void FStructSchemaBuilder::NoteMembersRecursively(const FBuiltStruct& Struct)
 			{
 				checkf(NotedSchema == Member.Schema, TEXT("Member '%s' in '%s' first added as '%s' and later as '%s'."),
 					*AllSchemas.GetDebug().Print(Member.Name),
-					*AllSchemas.GetDebug().Print(Declaration.Type),
+					*AllSchemas.GetDebug().Print(Declaration.Id),
 					*PrintMemberSchema(AllSchemas.GetIds(), NotedSchema),
 					*PrintMemberSchema(AllSchemas.GetIds(), Member.Schema));
 			}
@@ -351,7 +374,7 @@ void FStructSchemaBuilder::NoteRangeRecursively(ERangeSizeType NumType, TConstAr
 
 FBuiltStructSchema FStructSchemaBuilder::Build() const
 {
-	FType Type = Id == Declaration.Id ? Declaration.Type : AllSchemas.GetIds().Resolve(Id);
+	FType Type = AllSchemas.GetIds().Resolve(Id);
 	FBuiltStructSchema Out = { Type, Id, ToOptionalStruct(Declaration.Super), /* dense */ true };
 	
 	if (int32 Num = NotedMembers.CountSetBits())
diff --git a/Engine/Plugins/Experimental/PlainProps/Source/Private/PlainPropsDeclare.cpp b/Engine/Plugins/Experimental/PlainProps/Source/Private/PlainPropsDeclare.cpp
--- a/Engine/Plugins/Experimental/PlainProps/Source/Private/PlainPropsDeclare.cpp
+++ b/Engine/Plugins/Experimental/PlainProps/Source/Private/PlainPropsDeclare.cpp
@@ -1,8 +1,10 @@
 // Copyright Epic Games, Inc. All Rights Reserved.
 
 #include "PlainPropsDeclare.h"
+#include "PlainPropsSpecify.h"
 #include "Algo/Compare.h"
 #include "Containers/Set.h"
+#include "Templates/RefCounting.h"
 
 namespace PlainProps
 {
@@ -97,56 +99,14 @@ void CopyItems(T* It, TConstArrayView<T> Items)
 	}
 }
 
-static const FStructDeclaration& Declare(TArray<TUniquePtr<FStructDeclaration>>& Out, FDeclId Id, FType Type, uint16 Version, TConstArrayView<FMemberId> MemberOrder, EMemberPresence Occupancy, FOptionalDeclId Super, uint16 NumMembers)
+const FEnumDeclaration& FEnumDeclarations::Declare(FEnumId Id, FType Type, EEnumMode Mode, TConstArrayView<FEnumerator> Enumerators, EEnumAliases Policy)
 {
-	if (static_cast<int32>(Id.Idx) >= Out.Num())
+	if (static_cast<int32>(Id.Idx) >= Declarations.Num())
 	{
-		Out.SetNum(Id.Idx + 1);
-	}
-
-	TUniquePtr<FStructDeclaration>& Ptr = Out[Id.Idx];
-	if (Ptr)
-	{
-		check(Id == Ptr->Id);
-		check(Type == Ptr->Type);
-		check(Super == Ptr->Super);
-		check(Version == Ptr->Version);
-		check(NumMembers == Ptr->NumMembers);
-		check(Occupancy == Ptr->Occupancy);
-		check(Algo::Compare(MemberOrder, Ptr->GetMemberOrder()));
-
-		Ptr->RefCount++;
-	}
-	else
-	{
-		FStructDeclaration Header{1, Id, Type, Super, Version, NumMembers, Occupancy};
-		void* Data = FMemory::Malloc(sizeof(FStructDeclaration) + MemberOrder.Num() * MemberOrder.GetTypeSize());
-		Ptr.Reset(new (Data) FStructDeclaration(Header));
-		CopyItems(Ptr->MemberOrder, MemberOrder);
-	}
-	
-	return *Ptr;
-}
-
-const FStructDeclaration& FDeclarations::DeclareStruct(FDeclId Id, FType Type, uint16 Version, TConstArrayView<FMemberId> MemberOrder, EMemberPresence Occupancy, FOptionalDeclId Super)
-{
-	checkf(!(Super && Occupancy == EMemberPresence::RequireAll), TEXT("'%s' is a dense substruct, this isn't supported, see BuildSuperStruct()"), *Debug.Print(Id));
-	return Declare(/* out */ DeclaredStructs, Id, Type, Version, MemberOrder, Occupancy, Super, IntCastChecked<uint16>(MemberOrder.Num()));
-}
-
-const FStructDeclaration& FDeclarations::DeclareNumeralStruct(FDeclId Id, FType Type, TConstArrayView<FMemberId> Numerals, EMemberPresence Occupancy)
-{
-	return Declare(/* out */ DeclaredStructs, Id, Type, 0, Numerals, Occupancy, NoId, IntCastChecked<uint16>(Numerals.Num()));
-}
-
-const FEnumDeclaration& FDeclarations::DeclareEnum(FEnumId Id, FType Type, EEnumMode Mode, TConstArrayView<FEnumerator> Enumerators, EEnumAliases Policy)
-{
-	if (static_cast<int32>(Id.Idx) >= DeclaredEnums.Num())
-	{
-		DeclaredEnums.SetNum(Id.Idx + 1);
+		Declarations.SetNum(Id.Idx + 1);
 	}
 	
-	TUniquePtr<FEnumDeclaration>& Ptr = DeclaredEnums[Id.Idx];
+	TUniquePtr<FEnumDeclaration>& Ptr = Declarations[Id.Idx];
 	checkf(!Ptr, TEXT("'%s' is already declared"), *Debug.Print(Id));
 
 	TArray<FEnumerator, TInlineAllocator<64>> Tmp;
@@ -165,34 +125,114 @@ const FEnumDeclaration& FDeclarations::DeclareEnum(FEnumId Id, FType Type, EEnum
 	return *Ptr;
 }
 
-void FDeclarations::DropStructRef(FDeclId Id)
+#if DO_CHECK
+void FEnumDeclarations::Check(FEnumId Id) const
+{
+	checkf(Id.Idx < (uint32)Declarations.Num() && Declarations[Id.Idx], TEXT("'%s' is undeclared"), *Debug.Print(Id));
+}
+#endif
+
+////////////////////////////////////////////////////////////////////////////////////////////////
+
+uint32 FStructDeclaration::CalculateSize() const
+{
+	uint32 Out = sizeof(FStructDeclaration) + NumMembers * sizeof(FMemberId);
+	Out = Align(Out + NumInnerIds * sizeof(FInnerId), alignof(FInnerId));
+	Out = Align(Out + (NumMembers + NumInnerRanges) * sizeof(FMemberType), alignof(FMemberType));
+	return Out;
+}
+
+bool FStructDeclaration::Release() const
 {
-	Check(Id);
+	if (RefCount.fetch_sub(1, std::memory_order_acq_rel) == 1)
+	{
+		delete this;
+		return true;
+	}
+	return false;
+}
+
+////////////////////////////////////////////////////////////////////////////////////////////////
+
+FStructDeclarationPtr Declare(FStructSpec In) 
+{
+	check(In.MemberNames.Num() == In.MemberTypes.Num());
+
+	// Extract inner ids
+	TArray<FInnerId, TInlineAllocator<16>> InnerIds;
+	uint32 NumInnerRanges = 0;
+	for (FMemberSpecView Member : In.MemberTypes)
+	{
+		if (Member.InnermostId)
+		{
+			InnerIds.Emplace(Member.InnermostId.Get());
+		}
+		NumInnerRanges += Member.Ranges.Num();
+	}
+
+	// Make header and allocate
+	uint16 NumMembers = IntCastChecked<uint16>(In.MemberNames.Num());
+	uint16 NumInnerIds = IntCastChecked<uint16>(InnerIds.Num());
+	FStructDeclaration Header{1, In.Id, In.Super, In.Version, NumMembers, NumInnerIds, IntCastChecked<uint16>(NumInnerRanges), In.Occupancy};
+	FStructDeclaration* Out = new (FMemory::Malloc(Header.CalculateSize())) FStructDeclaration;
+	FMemory::Memcpy(Out, &Header, sizeof(Header));
 	
-	TUniquePtr<FStructDeclaration>& Ptr = DeclaredStructs[Id.Idx];
-	Ptr->RefCount--;
-	if (Ptr->RefCount == 0)
+	// Copy footer
+	CopyItems(Out->MemberNames, In.MemberNames);
+	CopyItems(const_cast<FInnerId*>(Out->GetInnerIds()), MakeConstArrayView(InnerIds));
+	FMemberType* TypeIt = const_cast<FMemberType*>(Out->GetTypes());
+	FMemberType* RangeIt = TypeIt + Header.NumMembers;
+	for (FMemberSpecView Member : In.MemberTypes)
 	{
-		Ptr.Reset();
+		if (int32 NumRanges = Member.Ranges.Num())
+		{
+			*TypeIt++ = FMemberType(Member.Ranges[0]);
+			for (ERangeSizeType Range : Member.Ranges.RightChop(1))
+			{
+				*RangeIt++ = FMemberType(Range);
+			}
+			*RangeIt++ = Member.InnermostType;
+		}
+		else
+		{
+			*TypeIt++ = Member.InnermostType;
+		}
 	}
+	check(TypeIt == Out->GetInnerRangeTypes());
+	check(RangeIt == Out->GetInnerRangeTypes() + NumInnerRanges);
+
+	// Make smart pointer
+	return FStructDeclarationPtr(Out, /* bAddRef */ false);	
 }
 
-const FStructDeclaration*
-FDeclarations::Find(FDeclId Id) const
+////////////////////////////////////////////////////////////////////////////////////////////////
+
+FMemberSpec::FMemberSpec(ERangeSizeType MaxSize, FMemberSpec Inner)
+: FMemberSpec(Inner)
 {
-	return Id.Idx < uint32(DeclaredStructs.Num()) ? DeclaredStructs[Id.Idx].Get() : nullptr;
+	RangeWrap(MaxSize);
 }
 
-#if DO_CHECK
-void FDeclarations::Check(FEnumId Id) const
+FMemberSpec::FMemberSpec(TConstArrayView<FMemberType> Members, FOptionalInnerId InInnermostId)
+: FMemberSpec(Members.Last(), InInnermostId)
 {
-	checkf(Id.Idx < (uint32)DeclaredEnums.Num() && DeclaredEnums[Id.Idx], TEXT("'%s' is undeclared"), *Debug.Print(Id));
+	for (FMemberType Range : Members.LeftChop(1))
+	{
+		RangeWrap(Range);
+	}
 }
 
-void FDeclarations::Check(FDeclId Id) const
+FMemberSpec::FMemberSpec(FMemberType Type, TConstArrayView<FMemberType> InnerRangeTypes, FOptionalInnerId InInnermostId)
+: FMemberSpec(InnerRangeTypes.IsEmpty() ? Type : InnerRangeTypes.Last(), InInnermostId)
 {
-	checkf(Id.Idx < (uint32)DeclaredStructs.Num() && DeclaredStructs[Id.Idx], TEXT("'%s' is undeclared"), *Debug.Print(Id));
+	if (int32 NumRanges = InnerRangeTypes.Num())
+	{
+		RangeWrap(Type);
+		for (FMemberType Nested : InnerRangeTypes.Slice(0, NumRanges - 1))
+		{
+			RangeWrap(Nested);	
+		}
+	}
 }
-#endif
 
 } // namespace PlainProps
\ No newline at end of file
diff --git a/Engine/Plugins/Experimental/PlainProps/Source/Private/PlainPropsDiff.cpp b/Engine/Plugins/Experimental/PlainProps/Source/Private/PlainPropsDiff.cpp
--- a/Engine/Plugins/Experimental/PlainProps/Source/Private/PlainPropsDiff.cpp
+++ b/Engine/Plugins/Experimental/PlainProps/Source/Private/PlainPropsDiff.cpp
@@ -114,9 +114,9 @@ bool DiffMember(/* in-out */ FMemberVisitor& It, const void* A, const void* B, C
 template<typename ContextType>
 bool DiffMembers(FBindId Id, const void* A, const void* B, ContextType& Ctx)
 {
-	const FSchemaBinding& Schema = Ctx.Schemas.GetStruct(Id);
-	const FStructDeclaration& Declaration = Ctx.Declarations.Get(Schema.DeclId);
-	
+	const FStructDeclaration* Declaration;
+	const FSchemaBinding& Schema = Ctx.Schemas.GetStruct(Id, /* out */ Declaration);
+
 	FMemberVisitor It(Schema);
 	if (Schema.HasSuper())
 	{
@@ -126,7 +126,7 @@ bool DiffMembers(FBindId Id, const void* A, const void* B, ContextType& Ctx)
 		}
 	}
 
-	for (FMemberId Name : Declaration.GetMemberOrder())
+	for (FMemberId Name : Declaration->GetMemberOrder())
 	{
 		if (DiffMember(/* in-out */It, A, B, Ctx, Name))
 		{
diff --git a/Engine/Plugins/Experimental/PlainProps/Source/Private/PlainPropsInternalBind.h b/Engine/Plugins/Experimental/PlainProps/Source/Private/PlainPropsInternalBind.h
--- a/Engine/Plugins/Experimental/PlainProps/Source/Private/PlainPropsInternalBind.h
+++ b/Engine/Plugins/Experimental/PlainProps/Source/Private/PlainPropsInternalBind.h
@@ -90,4 +90,17 @@ struct FMemberBinderBase
 	FRangeBinding* RangeBindingIt;
 };
 
+
+////////////////////////////////////////////////////////////////////////////////////////////////
+
+[[nodiscard]] inline FMemberType ToMemberType(FMemberBindType In)
+{
+	switch (In.GetKind())
+	{
+	case EMemberKind::Leaf:		return FMemberType(ToLeafType(In.AsLeaf()));
+	case EMemberKind::Range:	return FMemberType(In.AsRange());
+	default:					return FMemberType(In.AsStruct());
+	}
+}
+
 } // namespace PlainProps
\ No newline at end of file
diff --git a/Engine/Plugins/Experimental/PlainProps/Source/Private/PlainPropsInternalParse.h b/Engine/Plugins/Experimental/PlainProps/Source/Private/PlainPropsInternalParse.h
--- a/Engine/Plugins/Experimental/PlainProps/Source/Private/PlainPropsInternalParse.h
+++ b/Engine/Plugins/Experimental/PlainProps/Source/Private/PlainPropsInternalParse.h
@@ -9,6 +9,7 @@
 #include "Containers/ArrayView.h"
 #include "Containers/StringView.h"
 #include "Misc/Optional.h"
+#include "Templates/RefCounting.h"
 #include "Templates/Tuple.h"
 #include "Templates/UniquePtr.h"
 #include "Hash/xxhash.h"
@@ -167,14 +168,29 @@ FType MakeType(const FTypeTokens& TypeTokens, FTextIndexer& Names);
 
 ///////////////////////////////////////////////////////////////////////////////
 
+using FStructDeclarations = TArray<FStructDeclarationPtr>;
+
+struct FParseDeclarations final : IDeclarations
+{
+	explicit FParseDeclarations(FDebugIds Dbg) : Enums(Dbg) {}
+	
+	const FEnumDeclaration&				Get(FEnumId Id) { return Enums.Get(Id); }
+	const FStructDeclaration&			Get(FDeclId Id) { return *Structs[Id.Idx]; }
+
+	virtual FDeclId						Lower(FBindId Id) const override final;
+	virtual const FEnumDeclaration*		Find(FEnumId Id) const override { return &Enums.Get(Id); }
+	virtual const FStructDeclaration*	Find(FStructId Id) const override;
+
+	FEnumDeclarations					Enums;
+	FStructDeclarations					Structs;
+};
+
 class FBatchParser
 {
 public:
 	FBatchParser(FYamlTokenizer& InTokenizer, FScratchAllocator& InScratch)
-		: Tokenizer(InTokenizer), Scratch(InScratch), Types(FDebugIds(Names))
-	{ }
-
-	~FBatchParser() {};
+	: Tokenizer(InTokenizer), Scratch(InScratch), Types(FDebugIds(Names))
+	{}
 
 	void Parse(TArray64<uint8>& Out);
 
@@ -189,14 +205,13 @@ private:
 	FYamlTokenizer&									Tokenizer;
 	FScratchAllocator&								Scratch;
 	FTextIndexer									Names;
-	FDeclarations									Types;
+	FParseDeclarations								Types;
 	TArray<TTuple<FStructId, FBuiltStruct*>>		Objects;
 	FParsedSchemas									Schemas;
 	TArray64<FToken>								Tokens;
 	int64											EnumsIdx   = -1;
 	int64											StructsIdx = -1;
 	int64											ObjectsIdx = -1;
-
 };
 
-} // namespace PlainProps
+} // namespace PlainProps
\ No newline at end of file
diff --git a/Engine/Plugins/Experimental/PlainProps/Source/Private/PlainPropsInternalRead.h b/Engine/Plugins/Experimental/PlainProps/Source/Private/PlainPropsInternalRead.h
--- a/Engine/Plugins/Experimental/PlainProps/Source/Private/PlainPropsInternalRead.h
+++ b/Engine/Plugins/Experimental/PlainProps/Source/Private/PlainPropsInternalRead.h
@@ -82,10 +82,11 @@ struct TSchemaRange
 		static_assert(bStructSchema ^ std::is_convertible_v<T, FEnumSchema>);
 	}
 	
-	TSchemaIterator<T>				begin() const		{ return { Base, Offsets.begin() }; }
-	TSchemaIterator<T>				end() const			{ return { Base, Offsets.end() }; }
-	T&								First() const		{ return *reinterpret_cast<T*>(Base + Offsets[0]); }
-	T&								Last() const		{ return *reinterpret_cast<T*>(Base + Offsets.Last()); }
+	TSchemaIterator<T>				begin() const					{ return { Base, Offsets.begin() }; }
+	TSchemaIterator<T>				end() const						{ return { Base, Offsets.end() }; }
+	T&								First() const					{ return *reinterpret_cast<T*>(Base + Offsets[0]); }
+	T&								Last() const					{ return *reinterpret_cast<T*>(Base + Offsets.Last()); }
+	T&								operator[](uint32 Idx) const	{ return *reinterpret_cast<T*>(Base + Offsets[Idx]); }
 
 	uintptr_t						Base;
 	TConstArrayView<uint32>			Offsets;
diff --git a/Engine/Plugins/Experimental/PlainProps/Source/Private/PlainPropsLoad.cpp b/Engine/Plugins/Experimental/PlainProps/Source/Private/PlainPropsLoad.cpp
--- a/Engine/Plugins/Experimental/PlainProps/Source/Private/PlainPropsLoad.cpp
+++ b/Engine/Plugins/Experimental/PlainProps/Source/Private/PlainPropsLoad.cpp
@@ -535,7 +535,6 @@ struct FLoadPlanner
 {
 	// Constructor parameters
 	const FSchemaBatchId							BatchId;
-	const FDeclarations&							Declarations;
 	const FCustomBindings&							Customs;
 	const FSchemaBindings&							Schemas;
 	const TConstArrayView<FStructId>				RuntimeIds;
@@ -557,7 +556,7 @@ struct FLoadPlanner
 
 		if (Format == ESchemaFormat::InMemoryNames)
 		{
-			FLoadIdBinderDummy LoadIds{Declarations.GetDebug()};
+			FLoadIdBinderDummy LoadIds{Schemas.GetDebug()};
 			for (uint32 Idx = 0; Idx < NumPlans; ++Idx)
 			{
 				FLoadIdMapping Mapping;
@@ -571,7 +570,7 @@ struct FLoadPlanner
 		}
 	
 		TConstArrayView<FDeclId> DeclIds(static_cast<const FDeclId*>(RuntimeIds.GetData()), RuntimeIds.Num());
-		FLoadIdBinder LoadIds(DeclIds, Declarations.GetDebug());
+		FLoadIdBinder LoadIds(DeclIds, Schemas.GetDebug());
 		for (uint32 Idx = 0; Idx < NumPlans; ++Idx)
 		{
 			FLoadIdMapping Mapping;
@@ -603,7 +602,7 @@ struct FLoadPlanner
 
 		for (FStructSchemaId UnboundSaveId : UnboundSaveIds)
 		{
-			checkf(false, TEXT("Unbound struct '%s' can't be loaded"), *Declarations.GetDebug().Print(RuntimeIds[UnboundSaveId.Idx]));
+			checkf(false, TEXT("Unbound struct '%s' can't be loaded"), *Schemas.GetDebug().Print(RuntimeIds[UnboundSaveId.Idx]));
 		}
 
 		// Allocate load batch, copy plans and subset schemas, and fixup subset schema plans
@@ -815,10 +814,11 @@ private:
 
 		if (From.NumMembers)
 		{
-			if (const FSchemaBinding* To = Schemas.FindStruct(Mapping.Id))
+			const FStructDeclaration* Decl;
+			if (const FSchemaBinding* To = Schemas.FindStruct(Mapping.Id, /* out */ Decl))
 			{
 				// Possible optimization - some simple memcpy cases doesn't need to resolve the declaration
-				TConstArrayView<FMemberId> ToMemberIds = Declarations.Get(To->DeclId).GetMemberOrder();
+				TConstArrayView<FMemberId> ToMemberIds = Decl->GetMemberOrder();
 				return MakeLoadPlan(From, *To, ToMemberIds, RuntimeIds, LoadIds, /* out */ SubsetSchemaData);	
 			}
 			
@@ -831,9 +831,9 @@ private:
 	}
 };
 
-FLoadBatchPtr CreateLoadPlans(FSchemaBatchId BatchId, const FDeclarations& Declarations, const FCustomBindings& Customs, const FSchemaBindings& Schemas, TConstArrayView<FStructId> RuntimeIds, ESchemaFormat Format)
+FLoadBatchPtr CreateLoadPlans(FSchemaBatchId BatchId, const FCustomBindings& Customs, const FSchemaBindings& Schemas, TConstArrayView<FStructId> RuntimeIds, ESchemaFormat Format)
 {
-	return FLoadPlanner{BatchId, Declarations, Customs, Schemas, RuntimeIds}.CreatePlans(Format);
+	return FLoadPlanner{BatchId, Customs, Schemas, RuntimeIds}.CreatePlans(Format);
 }
 
 ////////////////////////////////////////////////////////////////////////////
diff --git a/Engine/Plugins/Experimental/PlainProps/Source/Private/PlainPropsParse.cpp b/Engine/Plugins/Experimental/PlainProps/Source/Private/PlainPropsParse.cpp
--- a/Engine/Plugins/Experimental/PlainProps/Source/Private/PlainPropsParse.cpp
+++ b/Engine/Plugins/Experimental/PlainProps/Source/Private/PlainPropsParse.cpp
@@ -556,11 +556,11 @@ private:
 class FEnumSchemaParser
 {
 public:
-	FEnumSchemaParser(FTokenReader& InTokens, FTextIndexer& InNames, FDeclarations& InTypes)
-		: ScopedReader(InTokens)
-		, Names(InNames)
-		, Types(InTypes)
-		{ }
+	FEnumSchemaParser(FTokenReader& InTokens, FTextIndexer& InNames, FEnumDeclarations& InEnums)
+	: ScopedReader(InTokens)
+	, Names(InNames)
+	, Declarations(InEnums)
+	{}
 
 	bool HasMore() const { return ScopedReader.HasMore(); }
 
@@ -576,7 +576,7 @@ private:
 
 	FScopedTokenReader	ScopedReader;
 	FTextIndexer&	  	Names;
-	FDeclarations&	  	Types;
+	FEnumDeclarations& 	Declarations;
 };
 
 ////////////////////////////////////////////////////////////////////////////////
@@ -584,10 +584,9 @@ private:
 class FStructSchemaParser
 {
 public:
-	FStructSchemaParser(
-		FTokenReader& InTokens, FTextIndexer& InNames, FDeclarations& InTypes, FParsedSchemas& InSchemas)
-		: ScopedReader(InTokens), Names(InNames), Types(InTypes), Schemas(InSchemas)
-	{ }
+	FStructSchemaParser(FTokenReader& InTokens, FTextIndexer& InNames, FStructDeclarations& InStructs, FParsedSchemas& InSchemas)
+	: ScopedReader(InTokens), Names(InNames), Declarations(InStructs), Schemas(InSchemas)
+	{}
 
 	bool HasMore() const { return ScopedReader.HasMore(); }
 
@@ -596,17 +595,17 @@ public:
 private:
 	bool ParseMemberSchema(FParsedMemberSchema& Out, FUtf8StringView String);
 
-	FParsedStructSchema DeclareStructSchema(
+	FParsedStructSchema ParseMembersAndDeclare(
 		const FTypeTokens& TypeTokens,
-		uint16 Version,
 		const TOptional<FTypeTokens>& ParsedSuper,
+		uint16 Version,
 		TConstArrayView<FUtf8StringView> MemberNames,
 		TConstArrayView<FUtf8StringView> MemberTypes);
 
-	FScopedTokenReader	ScopedReader;
-	FTextIndexer&	  	Names;
-	FDeclarations&	  	Types;
-	FParsedSchemas&		Schemas;
+	FScopedTokenReader		ScopedReader;
+	FTextIndexer&	  		Names;
+	FStructDeclarations&	Declarations;
+	FParsedSchemas&			Schemas;
 };
 
 ////////////////////////////////////////////////////////////////////////////////
@@ -619,7 +618,7 @@ public:
 		const FParsedStructSchema& Schema,
 		FScopedTokenReader InMemberTokens,
 		FTextIndexer& InNames,
-		FDeclarations& InTypes,
+		FParseDeclarations& InTypes,
 		FParsedSchemas& InSchemas);
 
 	FBuiltStruct* ParseMembers(FDeclId Id);
@@ -660,7 +659,7 @@ private:
 	TConstArrayView<FParsedMemberSchema>	MemberSchemas;
 	FScopedTokenReader						MemberTokens;
 	FTextIndexer&							Names;
-	FDeclarations&							Types;
+	FParseDeclarations&						Types;
 	FParsedSchemas&							Schemas;
 	FOptionalDeclId							UsedSuper;
 	uint32									TokenIdx = 0;
@@ -1079,15 +1078,16 @@ FParsedEnumSchema FEnumSchemaParser::DeclareEnumSchema(
 		Enumerators.Emplace(Names.MakeName(ConstantNames[I]), ConstantValues[I]);
 	}
 
-	Types.DeclareEnum(Id, Type, Mode, Enumerators, EEnumAliases::Strip);
+	const FEnumDeclaration& Decl = Declarations.Declare(Id, Type, Mode, Enumerators, EEnumAliases::Strip);
 
-	return { Id, Width, Types.Get(Id).GetEnumerators() };
+	return { Id, Width, Decl.GetEnumerators() };
 }
 
 ///////////////////////////////////////////////////////////////////////////////
 
 bool FStructSchemaParser::ParseStructSchema(FParsedStructSchema& Out)
 {
+	// Tokenize
 	FToken Token;
 	if (!GrabToNextBeginStruct(Token, ScopedReader))
 	{
@@ -1104,6 +1104,7 @@ bool FStructSchemaParser::ParseStructSchema(FParsedStructSchema& Out)
 	// Parse Struct Properties
 	TOptional<FTypeTokens> ParsedSuper;
 	uint16 TokenDepth = Token.Depth + 1;
+	uint16 Version = 0;
 	while (ScopedReader.Peek()->Depth == TokenDepth)
 	{
 		Token = ScopedReader.GrabToken();
@@ -1112,7 +1113,7 @@ bool FStructSchemaParser::ParseStructSchema(FParsedStructSchema& Out)
 			if (Token.Value() == GLiterals.Version)
 			{
 				Token = ScopedReader.GrabToken();
-				Out.Version = *Parse<uint16>(Token.Value());
+				Version = *Parse<uint16>(Token.Value());
 			}
 			else if (Token.Value() == GLiterals.DeclaredSuper)
 			{
@@ -1157,29 +1158,30 @@ bool FStructSchemaParser::ParseStructSchema(FParsedStructSchema& Out)
 		}
 	}
 
-	Out = DeclareStructSchema(TypeTokens.GetValue(), Out.Version, ParsedSuper, MemberNames, MemberTypes);
+	Out = ParseMembersAndDeclare(TypeTokens.GetValue(), ParsedSuper, Version, MemberNames, MemberTypes);
+
 	return true;
 }
 
-FParsedStructSchema FStructSchemaParser::DeclareStructSchema(
+FParsedStructSchema FStructSchemaParser::ParseMembersAndDeclare(
 	const FTypeTokens& TypeTokens,
-	uint16 Version,
 	const TOptional<FTypeTokens>& ParsedSuper,
+	uint16 Version,
 	TConstArrayView<FUtf8StringView> MemberNames,
 	TConstArrayView<FUtf8StringView> MemberTypes)
 {
-	const FType Type = MakeType(TypeTokens, Names);
-	const FDeclId Id = Names.IndexDeclId(Type);
+	const int32 Num = MemberNames.Num();
+	const FDeclId Id = Names.IndexDeclId(MakeType(TypeTokens, Names));
 	const FOptionalDeclId DeclaredSuper = ParsedSuper ? ToOptional(Names.IndexDeclId(MakeType(ParsedSuper.GetValue(), Names))) : NoId;
 
 	FOptionalDeclId UsedSuper;
 
 	TArray<FMemberId, TInlineAllocator<64>> MemberIds;
 	TArray<FParsedMemberSchema> MemberSchemas;
-	MemberIds.Reserve(MemberNames.Num());
-	MemberSchemas.Reserve(MemberNames.Num());
+	MemberIds.Reserve(Num);
+	MemberSchemas.Reserve(Num);
 
-	const int32 HasSuper = DeclaredSuper && MemberNames.Num() && MemberNames[0] == GLiterals.Super;
+	const int32 HasSuper = DeclaredSuper && Num && MemberNames[0] == GLiterals.Super;
 	if (HasSuper)
 	{
 		if (MemberTypes[0] != GLiterals.Dynamic)
@@ -1191,7 +1193,7 @@ FParsedStructSchema FStructSchemaParser::DeclareStructSchema(
 		}
 	}
 
-	for (int32 I = HasSuper, N = MemberNames.Num(); I < N; ++I)
+	for (int32 I = HasSuper; I < Num; ++I)
 	{
 		FParsedMemberSchema MemberSchema;
 		if (ParseMemberSchema(MemberSchema, MemberTypes[I]))
@@ -1205,9 +1207,22 @@ FParsedStructSchema FStructSchemaParser::DeclareStructSchema(
 		}
 	}
 
-	Types.DeclareStruct(Id, Type, Version, MemberIds, EMemberPresence::AllowSparse, DeclaredSuper);
+	TArray<FMemberSpec, TInlineAllocator<64>> MemberSpecs;
+	MemberSpecs.Reserve(MemberIds.Num());
+	for (const FParsedMemberSchema& Member : MemberSchemas)
+	{
+		MemberSpecs.Emplace(Member.Type, Member.InnerRangeTypes, Member.InnerSchema);
+	}
+
+	// Might want to switch to TMap
+	if (Id.Idx >= (uint32)Declarations.Num())
+	{
+		Declarations.SetNum(Id.Idx + 1);
+	}
+	check(!Declarations[Id.Idx]);
+	Declarations[Id.Idx] = Declare({Id, DeclaredSuper, Version, EMemberPresence::AllowSparse, MemberIds, MemberSpecs});
 
-	return { Id, Version, DeclaredSuper, UsedSuper, Types.Get(Id).GetMemberOrder(), MoveTemp(MemberSchemas) };
+	return { Id, Version, DeclaredSuper, UsedSuper, Declarations[Id.Idx]->GetMemberOrder(), MoveTemp(MemberSchemas) };
 }
 
 static TOptional<ERangeSizeType> GrabRangeSize(FUtf8StringView& String)
@@ -1297,7 +1312,7 @@ FMemberParser::FMemberParser(
 	const FParsedStructSchema& Schema,
 	FScopedTokenReader InMemberTokens,
 	FTextIndexer& InNames,
-	FDeclarations& InTypes,
+	FParseDeclarations& InTypes,
 	FParsedSchemas& InSchemas)
 : Scratch(InScratch)
 , MemberNames(Schema.MemberNames)
@@ -1950,13 +1965,13 @@ void FBatchParser::ParseStructSchemas(TConstArrayView64<FToken> TokensView)
 	FTokenReader TokenIt(TokensView);
 	verify(TokenIt.GrabToken().Value() == GLiterals.Structs);
 
-	FStructSchemaParser Parser(TokenIt, Names, Types, Schemas);
+	FStructSchemaParser Parser(TokenIt, Names, Types.Structs, Schemas);
 	while (Parser.HasMore())
 	{
 		FParsedStructSchema ParsedStructSchema;
 		if (Parser.ParseStructSchema(ParsedStructSchema))
 		{
-			Schemas.Structs.SetNum(Types.GetStructs().Num());
+			Schemas.Structs.SetNum(Types.Structs.Num());
 			Schemas.Structs[ParsedStructSchema.Id.Idx] = MoveTemp(ParsedStructSchema);
 		}
 	}
@@ -1967,7 +1982,7 @@ void FBatchParser::ParseEnumSchemas(TConstArrayView64<FToken> TokensView)
 	FTokenReader TokenIt(TokensView);
 	verify(TokenIt.GrabToken().Value() == GLiterals.Enums);
 
-	FEnumSchemaParser Parser(TokenIt, Names, Types);
+	FEnumSchemaParser Parser(TokenIt, Names, Types.Enums);
 	while (Parser.HasMore())
 	{
 		FParsedEnumSchema ParsedEnumSchema;
@@ -2018,26 +2033,31 @@ void FBatchParser::ParseObjects(TConstArrayView64<FToken> TokensView)
 	}
 }
 
-class FParseBindings : public IBindIds
+FDeclId FParseDeclarations::Lower(FBindId Id) const
 {
-	virtual FDeclId	Lower(FBindId Id) const override final
+	checkf(false, TEXT("All struct ids should be declared, nothing is bound with different names for text"));
+	return LowerCast(Id);
+}
+
+const FStructDeclaration* FParseDeclarations::Find(FStructId Id) const
+{
+	if (Id.Idx < (uint32)Structs.Num())
 	{
-		checkf(false, TEXT("All struct ids should be declared, nothing is bound with different names for text"));
-		return LowerCast(Id);
+		return Structs[Id.Idx];
 	}
-};
+	return nullptr;
+}
 
 void FBatchParser::Write(TArray64<uint8>& Out)
 {
-	FParseBindings Bindings;
-	FSchemasBuilder SchemaBuilders(Types, Names, Bindings, Scratch, ESchemaFormat::StableNames);
+	FSchemasBuilder SchemaBuilders(Names, Types, Scratch, ESchemaFormat::StableNames);
 	for (const TPair<FStructId, FBuiltStruct*>& Object : Objects)
 	{
 		SchemaBuilders.NoteStructAndMembers(Object.Key, *Object.Value);
 	}
 	FBuiltSchemas BuiltSchemas = SchemaBuilders.Build(); 
 
-	FWriter Writer(Names, Bindings, BuiltSchemas, ESchemaFormat::StableNames);
+	FWriter Writer(Names, Types, BuiltSchemas, ESchemaFormat::StableNames);
 	TArray64<uint8> Tmp;
 
 	// Write schemas
@@ -2069,4 +2089,4 @@ void ParseYamlBatch(TArray64<uint8>& OutBinary, FUtf8StringView Yaml)
 	BatchParser.Parse(OutBinary);
 }
 
-} // namespace PlainProps
+} // namespace PlainProps
\ No newline at end of file
diff --git a/Engine/Plugins/Experimental/PlainProps/Source/Private/PlainPropsSave.cpp b/Engine/Plugins/Experimental/PlainProps/Source/Private/PlainPropsSave.cpp
--- a/Engine/Plugins/Experimental/PlainProps/Source/Private/PlainPropsSave.cpp
+++ b/Engine/Plugins/Experimental/PlainProps/Source/Private/PlainPropsSave.cpp
@@ -205,18 +205,12 @@ FBuiltRange* SaveLeafRange(const void* Range, const ILeafRangeBinding& Binding,
 
 inline const FStructDeclaration& GetDeclaration(const FSaveContext& Ctx, FBindId BindId)
 {
-	// BindId is usually DeclId, unless it's been type-erased/lowered
-	if (const FStructDeclaration* Out = Ctx.Declarations.Find(LowerCast(BindId)))
-	{
-		return *Out;
-	}
-
 	if (const FStructDeclaration* Out = Ctx.Customs.FindDeclaration(BindId))
 	{
 		return *Out;
 	}
-	
-	return Ctx.Declarations.Get(Ctx.Schemas.GetStruct(BindId).DeclId);
+
+	return Ctx.Schemas.GetDeclaration(BindId);
 }
 
 [[nodiscard]] static FBuiltRange* SaveStructRange(const void* Range, const IItemRangeBinding& ItemBinding, const FSaveContext& Ctx, FBindId Id)
@@ -271,16 +265,6 @@ inline const FStructDeclaration& GetDeclaration(const FSaveContext& Ctx, FBindId
 	return SaveStructDelta(Struct, Default.Struct, Default.Id, Ctx);
 }
 
-[[nodiscard]] static FMemberType ToMemberType(FMemberBindType In)
-{
-	switch (In.GetKind())
-	{
-	case EMemberKind::Leaf:		return FMemberType(ToLeafType(In.AsLeaf()));
-	case EMemberKind::Range:	return FMemberType(In.AsRange());
-	default:					return FMemberType(In.AsStruct());
-	}
-}
-
 [[nodiscard]] static FMemberType* CreateInnerRangeTypes(FScratchAllocator& Scratch, uint32 NumInnerTypes, const FMemberBindType* InnerTypes)
 {
 	if (NumInnerTypes <= 1)
@@ -325,10 +309,10 @@ inline void SaveAllMembers(FMemberBuilder& Out, const void* Struct, const FSchem
 	{
 		FBindId SuperId = It.GrabSuper();
 		checkSlow(SuperId == ToOptionalStruct(Declaration.Super));
-		const FSchemaBinding& SuperSchema = Ctx.Schemas.GetStruct(SuperId);
-		const FStructDeclaration& SuperDecl = Ctx.Declarations.Get(SuperSchema.DeclId);
-		SaveAllMembers(Out, Struct, SuperSchema, SuperDecl, Ctx);
-		Out.BuildSuperStruct(Ctx.Scratch, SuperDecl, Ctx.Declarations.GetDebug());
+		const FStructDeclaration* SuperDecl;
+		const FSchemaBinding& SuperSchema = Ctx.Schemas.GetStruct(SuperId, /* out */ SuperDecl);
+		SaveAllMembers(Out, Struct, SuperSchema, *SuperDecl, Ctx);
+		Out.BuildSuperStruct(Ctx.Scratch, *SuperDecl, Ctx.Schemas.GetDebug());
 	}
 
 	for (FMemberId Name : Declaration.GetMemberOrder())
@@ -353,12 +337,11 @@ FBuiltStruct* SaveStruct(const void* Struct, FBindId Id, const FSaveContext& Ctx
 	}
 	else
 	{
-		const FSchemaBinding& Schema = Ctx.Schemas.GetStruct(Id);
-		Declaration = &Ctx.Declarations.Get(Schema.DeclId);	
+		const FSchemaBinding& Schema = Ctx.Schemas.GetStruct(Id, /* out */ Declaration);
 		SaveAllMembers(Out, Struct, Schema, *Declaration, Ctx);
 	}
 
-	return Out.BuildAndReset(Ctx.Scratch, *Declaration, Ctx.Declarations.GetDebug());
+	return Out.BuildAndReset(Ctx.Scratch, *Declaration, Ctx.Schemas.GetDebug());
 }
 
 ////////////////////////////////////////////////////////////////////////////////////////////////
@@ -437,10 +420,8 @@ static void SaveMemberDelta(FMemberBuilder& Out, const void* Struct, const void*
 	}
 }
 
-static const FStructDeclaration& SaveSchemaBoundStructDelta(FMemberBuilder& Out, const void* Struct, const void* Default, const FSchemaBinding& Schema, const FSaveContext& Ctx)
+static void SaveSchemaBoundStructDelta(FMemberBuilder& Out, const void* Struct, const void* Default, const FSchemaBinding& Schema,  const FStructDeclaration& Declaration, const FSaveContext& Ctx)
 {
-	const FStructDeclaration& Declaration = Ctx.Declarations.Get(Schema.DeclId);
-
 	if (Declaration.Occupancy == EMemberPresence::AllowSparse)
 	{
 		FMemberVisitor It(Schema);
@@ -448,9 +429,10 @@ static const FStructDeclaration& SaveSchemaBoundStructDelta(FMemberBuilder& Out,
 		{
 			FBindId SuperId = It.GrabSuper();
 			checkSlow(SuperId == ToOptionalStruct(Declaration.Super));
-			const FSchemaBinding& SuperSchema = Ctx.Schemas.GetStruct(SuperId);
-			const FStructDeclaration& SuperDecl = SaveSchemaBoundStructDelta(Out, Struct, Default, SuperSchema, Ctx);
-			Out.BuildSuperStruct(Ctx.Scratch, SuperDecl, Ctx.Declarations.GetDebug());
+			const FStructDeclaration* SuperDecl;
+			const FSchemaBinding& SuperSchema = Ctx.Schemas.GetStruct(SuperId, /* out */ SuperDecl);
+			SaveSchemaBoundStructDelta(Out, Struct, Default, SuperSchema, *SuperDecl, Ctx);
+			Out.BuildSuperStruct(Ctx.Scratch, *SuperDecl, Ctx.Schemas.GetDebug());
 		}
 
 		for (FMemberId Name : Declaration.GetMemberOrder())
@@ -468,8 +450,6 @@ static const FStructDeclaration& SaveSchemaBoundStructDelta(FMemberBuilder& Out,
 	{
 		SaveAllMembers(Out, Struct, Schema, Declaration, Ctx);
 	}
-
-	return Declaration;
 }
 
 FBuiltStruct* SaveStructDelta(const void* Struct, const void* Default, FBindId Id, const FSaveContext& Ctx)
@@ -481,14 +461,17 @@ FBuiltStruct* SaveStructDelta(const void* Struct, const void* Default, FBindId I
 		if (Custom->DiffCustom(Struct, Default, Ctx))
 		{
 			Custom->SaveCustom(Out, Struct, Default, Ctx);
+			// Could type-check saved custom struct against declaration here,
+			// or add a type-checking BuildAndReset overload and call that.
+			// It's not important to type-check bindings, they're generated.
 		}
 	}
 	else
 	{
-		const FSchemaBinding& Schema = Ctx.Schemas.GetStruct(Id);
-		Declaration = &SaveSchemaBoundStructDelta(Out, Struct, Default, Schema, Ctx);
+		const FSchemaBinding& Schema = Ctx.Schemas.GetStruct(Id, /* out */ Declaration);
+		SaveSchemaBoundStructDelta(Out, Struct, Default, Schema, *Declaration, Ctx);
 	}
-	return Out.BuildAndReset(Ctx.Scratch, *Declaration, Ctx.Declarations.GetDebug());
+	return Out.BuildAndReset(Ctx.Scratch, *Declaration, Ctx.Schemas.GetDebug());
 }
 
 FBuiltStruct* SaveStructDeltaIfDiff(const void* Struct, const void* Default, FBindId Id, const FSaveContext& Ctx)
@@ -500,16 +483,16 @@ FBuiltStruct* SaveStructDeltaIfDiff(const void* Struct, const void* Default, FBi
 		if (Custom->DiffCustom(Struct, Default, Ctx))
 		{
 			Custom->SaveCustom(Out, Struct, Default, Ctx);
-			return Out.BuildAndReset(Ctx.Scratch, *Declaration, Ctx.Declarations.GetDebug());
+			return Out.BuildAndReset(Ctx.Scratch, *Declaration, Ctx.Schemas.GetDebug());
 		}
 		
 		return nullptr;
 	}
 
-	const FSchemaBinding& Schema = Ctx.Schemas.GetStruct(Id);
-	Declaration = &SaveSchemaBoundStructDelta(Out, Struct, Default, Schema, Ctx);
+	const FSchemaBinding& Schema = Ctx.Schemas.GetStruct(Id, /* out */ Declaration);
+	SaveSchemaBoundStructDelta(Out, Struct, Default, Schema, *Declaration, Ctx);
 
-	return Out.IsEmpty() ? nullptr : Out.BuildAndReset(Ctx.Scratch, *Declaration, Ctx.Declarations.GetDebug());
+	return Out.IsEmpty() ? nullptr : Out.BuildAndReset(Ctx.Scratch, *Declaration, Ctx.Schemas.GetDebug());
 }
 
 ////////////////////////////////////////////////////////////////////////////////////////////////
diff --git a/Engine/Plugins/Experimental/PlainProps/Source/Private/PlainPropsUeCoreBindings.cpp b/Engine/Plugins/Experimental/PlainProps/Source/Private/PlainPropsUeCoreBindings.cpp
--- a/Engine/Plugins/Experimental/PlainProps/Source/Private/PlainPropsUeCoreBindings.cpp
+++ b/Engine/Plugins/Experimental/PlainProps/Source/Private/PlainPropsUeCoreBindings.cpp
@@ -52,9 +52,14 @@ void FTransformBinding::Save(FMemberBuilder& Dst, const FTransform& Src, const F
 {
 	static_assert(std::is_same_v<decltype(FTransform().GetTranslation().X), double>);
 	
-	const FStructDeclaration& VectorDecl = Ctx.Declarations.Get(LowerCast(VectorId));
-	const FStructDeclaration& QuatDecl = Ctx.Declarations.Get(LowerCast(QuatId));
-	FDenseMemberBuilder Inner = { Ctx.Scratch, Ctx.Declarations.GetDebug() };
+	static const struct
+	{
+		FStructDeclarationPtr Vector;
+		FStructDeclarationPtr Quat;
+	}
+	Decls = { &Ctx.Schemas.GetDeclaration(VectorId), &Ctx.Schemas.GetDeclaration(QuatId) };
+
+	FDenseMemberBuilder Inner = { Ctx.Scratch, Ctx.Schemas.GetDebug() };
 
 	FVector T = Src.GetTranslation();
 	FQuat R = Src.GetRotation();
@@ -64,24 +69,24 @@ void FTransformBinding::Save(FMemberBuilder& Dst, const FTransform& Src, const F
 	{
 		if (T != Default->GetTranslation())
 		{
-			Dst.AddStruct(MemberIds[(uint8)EMember::Translate], VectorId, Inner.BuildHomogeneous(VectorDecl, T.X, T.Y, T.Z));
+			Dst.AddStruct(MemberIds[(uint8)EMember::Translate], VectorId, Inner.BuildHomogeneous(*Decls.Vector, T.X, T.Y, T.Z));
 		}
 
 		if (R != Default->GetRotation())
 		{
-			Dst.AddStruct(MemberIds[(uint8)EMember::Rotate], QuatId, Inner.BuildHomogeneous(QuatDecl, R.X, R.Y, R.Z, R.W));
+			Dst.AddStruct(MemberIds[(uint8)EMember::Rotate], QuatId, Inner.BuildHomogeneous(*Decls.Quat, R.X, R.Y, R.Z, R.W));
 		}
 
 		if (S != Default->GetScale3D())
 		{
-			Dst.AddStruct(MemberIds[(uint8)EMember::Scale], VectorId, Inner.BuildHomogeneous(VectorDecl, S.X, S.Y, S.Z));
+			Dst.AddStruct(MemberIds[(uint8)EMember::Scale], VectorId, Inner.BuildHomogeneous(*Decls.Vector, S.X, S.Y, S.Z));
 		}
 	}
 	else
 	{	
-		Dst.AddStruct(MemberIds[(uint8)EMember::Translate], VectorId, Inner.BuildHomogeneous(VectorDecl, T.X, T.Y, T.Z));		
-		Dst.AddStruct(MemberIds[(uint8)EMember::Rotate], QuatId, Inner.BuildHomogeneous(QuatDecl, R.X, R.Y, R.Z, R.W));
-		Dst.AddStruct(MemberIds[(uint8)EMember::Scale], VectorId, Inner.BuildHomogeneous(VectorDecl, S.X, S.Y, S.Z));
+		Dst.AddStruct(MemberIds[(uint8)EMember::Translate], VectorId, Inner.BuildHomogeneous(*Decls.Vector, T.X, T.Y, T.Z));		
+		Dst.AddStruct(MemberIds[(uint8)EMember::Rotate], QuatId, Inner.BuildHomogeneous(*Decls.Quat, R.X, R.Y, R.Z, R.W));
+		Dst.AddStruct(MemberIds[(uint8)EMember::Scale], VectorId, Inner.BuildHomogeneous(*Decls.Vector, S.X, S.Y, S.Z));
 	}
 }
 
diff --git a/Engine/Plugins/Experimental/PlainProps/Source/Private/PlainPropsWrite.cpp b/Engine/Plugins/Experimental/PlainProps/Source/Private/PlainPropsWrite.cpp
--- a/Engine/Plugins/Experimental/PlainProps/Source/Private/PlainPropsWrite.cpp
+++ b/Engine/Plugins/Experimental/PlainProps/Source/Private/PlainPropsWrite.cpp
@@ -18,11 +18,11 @@ namespace PlainProps
 // Rewrite as a more compact data structure once we get a large number of ids
 struct FWriteIds
 {
-	FWriteIds(const FIdIndexerBase& DeclaredIds, const IBindIds& BindIds, const FBuiltSchemas& Schemas, ESchemaFormat Format);
+	FWriteIds(const FIdIndexerBase& DeclaredIds, const IDeclarations& Decls, const FBuiltSchemas& Schemas, ESchemaFormat Format);
 
 	bool								HasStableNames() const { return !Names.IsEmpty(); }
 
-	const IBindIds&						BindIds;
+	const IDeclarations&				Declarations;
 
 	TArray<FOptionalNameId>				Names;
 	TArray<FOptionalNestedScopeId>		NestedScopes;
@@ -222,8 +222,8 @@ static void CopyUsed(TArray<ItemType>& Out, const TArray<IdType>& Kept, IndexerT
 	}
 }
 
-FWriteIds::FWriteIds(const FIdIndexerBase& Ids, const IBindIds& InBindIds, const FBuiltSchemas& Schemas, ESchemaFormat Format)
-: BindIds(InBindIds)
+FWriteIds::FWriteIds(const FIdIndexerBase& Ids, const IDeclarations& Decls, const FBuiltSchemas& Schemas, ESchemaFormat Format)
+: Declarations(Decls)
 , Debug(Ids)
 {
 	Structs.Init(NoId, Ids.NumStructs());
@@ -298,9 +298,8 @@ FStructSchemaId FWriteIds::RemapStruct(FStructId OldId) const
 	}
 
 	// Could optimize by caching Structs[OldBindId.Idx] here
-	FDeclId OldDeclId = BindIds.Lower(FBindId(OldId));
+	FDeclId OldDeclId = Declarations.Lower(FBindId(OldId));
 	return Structs[OldDeclId.Idx].Get();
-
 }
 
 //////////////////////////////////////////////////////////////////////////
@@ -820,10 +819,10 @@ private:
 
 ////////////////////////////////////////////////////////////////////////////////////////////////////////////////////////////////////////////////////
 
-FWriter::FWriter(const FIdIndexerBase& AllIds, const IBindIds& BindIds, const FBuiltSchemas& InSchemas, ESchemaFormat Format)
+FWriter::FWriter(const FIdIndexerBase& AllIds, const IDeclarations& Declarations, const FBuiltSchemas& InSchemas, ESchemaFormat Format)
 : Schemas(InSchemas)
 , Debug(AllIds)
-, NewIds(new FWriteIds(AllIds, BindIds, InSchemas, Format))
+, NewIds(new FWriteIds(AllIds, Declarations, InSchemas, Format))
 {}
 
 FWriter::~FWriter()
diff --git a/Engine/Plugins/Experimental/PlainProps/Source/Public/PlainPropsBind.h b/Engine/Plugins/Experimental/PlainProps/Source/Public/PlainPropsBind.h
--- a/Engine/Plugins/Experimental/PlainProps/Source/Public/PlainPropsBind.h
+++ b/Engine/Plugins/Experimental/PlainProps/Source/Public/PlainPropsBind.h
@@ -11,6 +11,7 @@
 #include "PlainPropsDeclare.h"
 #include "PlainPropsTypename.h"
 #include "PlainPropsRead.h"
+#include "PlainPropsSpecify.h"
 #include "PlainPropsTypes.h"
 #include <tuple>
 
@@ -28,7 +29,6 @@ class FMemberBuilder;
 class FRangeLoader;
 struct FSchemaBatch;
 class FScratchAllocator;
-class FStructBinding;
 struct FStructLoadView;
 class FRangeBinding;
 struct FSaveContext;
@@ -233,8 +233,6 @@ struct FBothStructId
 	bool					IsLowered() const { return BindId != DeclId; }
 };
 
-////////////////////////////////////////////////////////////////////////////////////////////////
-
 // FStructId statically known to share FBindId and FDeclId, i.e. not lowered
 struct FDualStructId : FStructId
 {
@@ -286,42 +284,42 @@ struct FIdWindow
 };
 
 // Slightly optimized map from FStructId to ICustomBinding/FStructDeclaration/FInnersHandle
-class FCustomBindingMap
+struct FCustomBindingMap
 {
-public:
 	UE_NONCOPYABLE(FCustomBindingMap);
-	explicit FCustomBindingMap(FDebugIds Dbg) : Debug(Dbg) {}
-	~FCustomBindingMap() { FMemory::Free(Values); }
-
-	void							Bind(FBindId Id, ICustomBinding& Binding, const FStructDeclaration& Declaration, FInnersHandle LoweredInners);
-	FCustomBindingHandle			Find(FBindId Id) const;
-	void							Drop(FBindId Id);
-
-private:
-	FIdWindow						Window;
-	TSet<FStructId>			Keys;
-	ICustomBinding**				Values = nullptr;
-	uint32							MaxValues = 0;
-public:
-	FDebugIds						Debug;
+	explicit FCustomBindingMap(FDebugIds Dbg);
+	~FCustomBindingMap();
+
+	void								Bind(FBindId Id, ICustomBinding& Binding, FStructDeclarationPtr&& Declaration, FInnersHandle LoweredInners);
+	FCustomBindingHandle				Find(FBindId Id) const;
+	void								Drop(FBindId Id);
+
+	FIdWindow							Window;
+	TSet<FBindId>						Keys;
+	ICustomBinding**					Values = nullptr;
+	uint32								MaxValues = 0;
+	TArray<FStructDeclarationPtr>		Declarations; // Consider switch to TSparseArray or TMap
+	FDebugIds							Debug;
 };
 
 class FCustomBindings
 {
 public:
 	// @param Binding must outlive this or call DropStruct()
-	PLAINPROPS_API void							BindStruct(FBindId Id, ICustomBinding& Binding, const FStructDeclaration& Declaration, TConstArrayView<FInnerStruct> LoweredInners = {});
+	PLAINPROPS_API void							BindStruct(FBindId Id, ICustomBinding& Binding, FStructDeclarationPtr&& Declaration, TConstArrayView<FInnerStruct> LoweredInners);
+	PLAINPROPS_API void							BindStruct(FBindId Id, ICustomBinding& Binding, const FStructSpec& DeclSpec, TConstArrayView<FInnerStruct> LoweredInners);
 	PLAINPROPS_API const ICustomBinding*		FindStruct(FBindId Id) const;
 	PLAINPROPS_API const ICustomBinding*		FindStruct(FBindId Id, TConstArrayView<FInnerStruct>& OutLoweredInners) const;
+	PLAINPROPS_API const ICustomBinding*		FindStruct(FBindId Id, const FStructDeclaration*& OutDeclaration) const;
 	PLAINPROPS_API ICustomBinding*				FindStructToSave(FBindId Id, const FStructDeclaration*& OutDeclaration) const;
 	PLAINPROPS_API const FStructDeclaration*	FindDeclaration(FBindId Id) const;
 	PLAINPROPS_API void							DropStruct(FBindId Id);
 
 	virtual FCustomBindingHandle				Find(FBindId Id) const = 0;
 protected:
-	FCustomBindings(TArray<FInnerStruct>& Inners, FDebugIds Dbg) : Map(Dbg), BottomInners(Inners) {}
-	FCustomBindings(const FCustomBindings* Under) : Map(Under->Map.Debug), BottomInners(Under->BottomInners) {}
-	~FCustomBindings() = default;
+	PLAINPROPS_API FCustomBindings(TArray<FInnerStruct>& Inners, FDebugIds Dbg);
+	PLAINPROPS_API FCustomBindings(const FCustomBindings* Under);
+	PLAINPROPS_API ~FCustomBindings();
 
 	FCustomBindingMap							Map;
 	TArray<FInnerStruct>&						BottomInners;
@@ -375,21 +373,20 @@ class FConstructionRequest
 {
 	friend FRangeLoader;
 
-	void* const Range = nullptr;
-	const uint64 Num = 0;
-	uint64 Index = 0;
+	void* const		Range = nullptr;
+	const uint64	Num = 0;
+	uint64			Index = 0;
 
 	FConstructionRequest(void* InRange, uint64 InNum) : Range(InRange), Num(InNum) {}
-	
+
 public:
 	template<typename T>
-	T& GetRange() const { return *reinterpret_cast<T*>(Range); }
-	
-	uint64 NumTotal() const { return Num; }
-	uint64 NumMore() const { return Num - Index; }
-	uint64 GetIndex() const { return Index; }
-	bool IsFirstCall() const { return Index == 0;}
-	bool IsFinalCall() const { return Index == Num;}
+	T&				GetRange() const		{ return *reinterpret_cast<T*>(Range); }
+	uint64			NumTotal() const		{ return Num; }
+	uint64			NumMore() const			{ return Num - Index; }
+	uint64			GetIndex() const		{ return Index; }
+	bool			IsFirstCall() const		{ return Index == 0;}
+	bool			IsFinalCall() const		{ return Index == Num;}
 };
 
 // Response from IItemRangeBinding with contiguous items ready to be loaded
@@ -620,22 +617,33 @@ struct FMemberBinding
 	PLAINPROPS_API FBothType		IndexParameterName(FIdIndexerBase& Ids) const;
 };
 
-class FSchemaBindings : public IBindIds
+class FSchemaBindings
 {
 public:
 	UE_NONCOPYABLE(FSchemaBindings);
-	explicit FSchemaBindings(FDebugIds In) : Debug(In) {}
+	PLAINPROPS_API explicit FSchemaBindings(FDebugIds In);
 	PLAINPROPS_API ~FSchemaBindings();
+	
+	PLAINPROPS_API void							BindStruct(FBindId Id, TConstArrayView<FMemberBinding> Members, FStructSpec Decl);
+	PLAINPROPS_API void							BindStruct(FBindId Id, TConstArrayView<FMemberBinding> Members, FStructDeclarationPtr&& Decl);
+	PLAINPROPS_API const FSchemaBinding*		FindStruct(FBindId Id) const;
+	PLAINPROPS_API const FSchemaBinding*		FindStruct(FBindId Id, const FStructDeclaration*& OutDeclaration) const;
+	PLAINPROPS_API const FSchemaBinding&		GetStruct(FBindId Id) const;
+	PLAINPROPS_API const FSchemaBinding&		GetStruct(FBindId Id, const FStructDeclaration*& OutDecl) const;
+	PLAINPROPS_API void							DropStruct(FBindId Id);
 
-	PLAINPROPS_API void						BindStruct(FBindId Id, FDeclId DeclId, TConstArrayView<FMemberBinding> Schema);
-	PLAINPROPS_API const FSchemaBinding*	FindStruct(FBindId Id) const;
-	PLAINPROPS_API const FSchemaBinding&	GetStruct(FBindId Id) const;
-	PLAINPROPS_API void						DropStruct(FBindId Id);
+	PLAINPROPS_API FDeclId						Lower(FBindId Id) const;
+	PLAINPROPS_API const FStructDeclaration&	GetDeclaration(FBindId Id) const;
+	PLAINPROPS_API const FStructDeclaration*	FindDeclaration(FBindId Id) const;
+	FDebugIds									GetDebug() const { return Debug; }
 
-	PLAINPROPS_API virtual FDeclId			Lower(FBindId Id) const override final;
 private:
-	TArray<TUniquePtr<FSchemaBinding>>		Bindings;
-	FDebugIds								Debug;
+	TArray<FStructDeclarationPtr>				Declarations;
+	TArray<TUniquePtr<FSchemaBinding>>			Bindings;
+	FDebugIds									Debug;
+
+	FSchemaBinding*								Bind(FBindId Id, TConstArrayView<FMemberBinding> Members, FDeclId DeclId);
+	FStructDeclarationPtr&						At(FDeclId Id);
 };
 
 ////////////////////////////////////////////////////////////////////////////////////////////////
@@ -650,24 +658,23 @@ struct IDefaultStructs
 
 struct FBindContext
 {
-	const FDeclarations&		Declarations;
-	const FSchemaBindings&		Schemas;
-	FCustomBindings&			Customs;
-};
+	const FSchemaBindings&								Schemas;
+	FCustomBindings&									Customs;
 
-////////////////////////////////////////////////////////////////////////////////////////////////
+	PLAINPROPS_API const FStructDeclaration&			GetDeclaration(FBindId Id) const;
+};
 
-struct FStructBindIds final : IBindIds
+struct FBindDeclarations final : IDeclarations
 {
-	FStructBindIds(const FCustomBindings& InCustoms, const FSchemaBindings& InSchemas)
-	: Customs(InCustoms)
-	, Schemas(InSchemas)
-	{}
-
-	const FCustomBindings& Customs;
-	const FSchemaBindings& Schemas;
+	FBindDeclarations(const FEnumDeclarations& E, const FCustomBindings& C, const FSchemaBindings& S) : Enums(E), Customs(C), Schemas(S) {}
 
-	PLAINPROPS_API virtual FDeclId Lower(FBindId Id) const override;
+	const FEnumDeclarations&							Enums;
+	const FCustomBindings&								Customs;
+	const FSchemaBindings&								Schemas;
+	
+	PLAINPROPS_API virtual const FEnumDeclaration*		Find(FEnumId Id) const override;
+	PLAINPROPS_API virtual const FStructDeclaration*	Find(FStructId Id) const override;
+	PLAINPROPS_API virtual FDeclId						Lower(FBindId Id) const override;
 };
 
 ////////////////////////////////////////////////////////////////////////////////////////////////
@@ -715,20 +722,6 @@ FType IndexStructName()
 	}
 }
 
-template<class Ids, class Typename>
-FBindId IndexStructBindIdIfNeeded(FDeclId DeclId)
-{
-	if constexpr (ExplicitBindName<Typename> || ParametricName<Typename>)
-	{
-		// Note could pass in and reuse declared namespace here
-		return FBindId(Ids::IndexStruct(IndexStructName<Ids, ETypename::Bind, Typename>()));
-	}
-	else
-	{
-		return UpCast(DeclId);
-	}
-}
-
 template<class Ids, class Typename>
 FBothStructId IndexStructBothId(FType DeclName = IndexStructName<Ids, ETypename::Decl, Typename>())
 {	
@@ -745,6 +738,12 @@ FBothStructId IndexStructBothId(FType DeclName = IndexStructName<Ids, ETypename:
 	return Out;
 }
 
+template<class Ids, class Typename>
+FDualStructId IndexStructDualId(FType Name = IndexStructName<Ids, ETypename::Decl, Typename>()) requires (!ExplicitBindName<Typename>)
+{	
+	return FDualStructId(Ids::IndexStruct(Name));
+}
+
 // Cached by function static
 template<class Ids, typename Struct>
 FDeclId GetStructDeclId()
@@ -769,6 +768,14 @@ FBothStructId GetStructBothId()
 	return Id;
 }
 
+// Cached by function static
+template<class Ids, typename Struct>
+FDualStructId GetStructDualId()
+{
+	static FDualStructId Id = IndexStructDualId<Ids, TTypename<Struct>>();
+	return Id;
+}
+
 template<class Ids, typename Ctti>
 FType IndexCttiName()
 {
@@ -782,7 +789,7 @@ FType IndexCttiName()
 }
 
 // Cached by function static
-template<class Ids, typename Enum>
+template<class Ids, Enumeration Enum>
 FEnumId GetEnumId()
 {
 	static FEnumId Id = Ids::IndexEnum(IndexCttiName<Ids, CttiOf<Enum>>());
@@ -854,16 +861,19 @@ using InnerStructArray = TArray<FInnerStruct, TInlineAllocator<8>>;
 
 struct FCustomInit
 {
-	InnerStructArray& OutLowered;
+	UE_NONCOPYABLE(FCustomInit);
+	FCustomInit() = default;
+
+	InnerStructArray				LoweredInners;
 
 	// Only needed for inner structs that might get type-erased/lowered
-	void RegisterInnerStruct(FBothStructId Id, TConstArrayView<FMemberId> Names)
+	void RegisterInnerStruct(FBothStructId InnerId, TConstArrayView<FMemberId> Names)
 	{
-		if (Id.IsLowered())
+		if (InnerId.IsLowered())
 		{
 			for (FMemberId Name : Names)
 			{
-				OutLowered.Add({Name, Id.BindId});
+				LoweredInners.Add({Name, InnerId.BindId});
 			}
 		}
 	}
@@ -874,6 +884,31 @@ template<class Ids>
 struct TCustomInit : FCustomInit, TInit<Ids>
 {};
 
+// Constructor API for static custom bindings to create ids, register type-erased inner structs
+// and specify member types, inheritance and version 
+template<class Ids, uint32 N>
+struct TCustomSpecifier : TCustomInit<Ids>
+{
+	FOptionalStructId				Super;
+	uint16							Version = 0;
+	FMemberSpec						Members[N];
+
+	template <typename... Ts>
+    void SetMembers(Ts&&... Specs) requires (N == sizeof...(Specs))
+    {
+		FMemberSpec* It = &Members[0];
+		((*It++ = FMemberSpec(Forward<Ts>(Specs))), ...);
+    }
+
+	void FillMembers(FMemberSpec Same)
+	{
+		for (FMemberSpec& Member : Members)
+		{
+			Member = Same;
+		}
+	}
+};
+
 // Workaround for clang requiring FStructLoadView definition, make it dependent on any other type
 template<typename>
 struct TDependent { using StructLoadView = FStructLoadView; };
@@ -884,31 +919,31 @@ struct TScopedStructBinding : FBothStructId, CustomBinding
 {
 	using Ids = typename Runtime::Ids;
 	using Typename = TCustomTypename<CustomBinding>::Type;
+	static constexpr uint32 N = sizeof(CustomBinding::MemberIds) / sizeof(CustomBinding::MemberIds[0]);
 
-	explicit TScopedStructBinding(FType DeclName = IndexStructName<Ids, ETypename::Decl, Typename>(), InnerStructArray Lowered = {})
-	: TScopedStructBinding(DeclName, IndexStructBothId<Ids, Typename>(DeclName), Lowered)
+	explicit TScopedStructBinding(FType DeclName = IndexStructName<Ids, ETypename::Decl, Typename>())
+	: TScopedStructBinding(IndexStructBothId<Ids, Typename>(DeclName))
 	{}
 
-	explicit TScopedStructBinding(FBothStructId Id, InnerStructArray Lowered = {})
-	: TScopedStructBinding(Ids::GetIndexer().Resolve(Id.DeclId), Id, Lowered)
+	explicit TScopedStructBinding(FDualStructId Dual)
+	: TScopedStructBinding(FBothStructId{Dual, Dual})
 	{}
 
-	explicit TScopedStructBinding(FDualStructId Id, InnerStructArray Lowered = {})
-	: TScopedStructBinding(Ids::GetIndexer().Resolve(Id), FBothStructId{Id, Id}, Lowered)
+	explicit TScopedStructBinding(FBothStructId Both)
+	: TScopedStructBinding(Both, TCustomSpecifier<Ids, N>())
 	{}
 
-	TScopedStructBinding(FType DeclName, FBothStructId BothId, InnerStructArray& Lowered)
-	: FBothStructId(BothId)
-	, CustomBinding(TCustomInit<Ids>{{Lowered}})
+	TScopedStructBinding(FBothStructId Both, TCustomSpecifier<Ids, N>&& Init)
+	: FBothStructId(Both)
+	, CustomBinding(/* out */ Init)
 	{
-		const FStructDeclaration& Declaration = Runtime::GetTypes().DeclareStruct(DeclId, DeclName, 0, CustomBinding::MemberIds, TOccupancyOf<T>::Value);
-		Runtime::GetCustoms().BindStruct(BindId, *this, Declaration, Lowered);
+		FStructSpec Decl = { DeclId, /* todo */ NoId, Init.Version, TOccupancyOf<T>::Value, CustomBinding::MemberIds, Init.Members };
+		Runtime::GetCustoms().BindStruct(BindId, *this, Decl, Init.LoweredInners);
 	}
 
 	~TScopedStructBinding()
 	{
 		Runtime::GetCustoms().DropStruct(BindId);
-		Runtime::GetTypes().DropStructRef(DeclId);
 	}
 
 	virtual void SaveCustom(FMemberBuilder& Dst, const void* Src, const void* Default, const FSaveContext& Ctx) override
@@ -934,73 +969,63 @@ struct TScopedStructBinding : FBothStructId, CustomBinding
 
 // Specialization for schemabound structs
 template<class T, class Runtime>
-struct TScopedStructBinding<T, Runtime, void>
+struct TScopedStructBinding<T, Runtime, void> : FBothStructId
 {
 	using Ids = typename Runtime::Ids;
 	using Typename = TTypename<T>;
 	
-	const FDeclId DeclId;
-	const FBindId BindId;
-
-	TScopedStructBinding(FType DeclName = IndexStructName<Ids, ETypename::Decl, Typename>())
-	: DeclId(DeclareNativeStruct<CttiOf<T>, Ids>(Runtime::GetTypes(), TOccupancyOf<T>::Value))
-	, BindId(IndexStructBindIdIfNeeded<Ids, Typename>(DeclId))
+	TScopedStructBinding(FBothStructId Both = IndexStructBothId<Ids, Typename>() )
+	: FBothStructId(Both)
 	{
-		BindNativeStruct<CttiOf<T>, Runtime>(Runtime::GetSchemas(), BindId, DeclId);
+		BindNativeStruct<CttiOf<T>, Runtime>(/* out */ Runtime::GetSchemas(), Both, TOccupancyOf<T>::Value);
 	}
 
 	~TScopedStructBinding()
 	{
 		Runtime::GetSchemas().DropStruct(BindId);
-		Runtime::GetTypes().DropStructRef(DeclId);
 	}
 };
 
 template<class Struct, typename CustomBinding, class Runtime>
-FBindId BindCustomStructOnce()
+FBothStructId BindCustomStructOnce()
 {
 	static TScopedStructBinding<Struct, Runtime, CustomBinding> Instance;
-	return Instance.BindId;
+	return Instance;
 }
 
 inline constexpr FMemberBindType DefaultStructBindType = FMemberBindType(FStructType{EMemberKind::Struct, /* IsDynamic */ 0, /* IsSuper */ 0});
 inline constexpr FMemberBindType SuperStructBindType = FMemberBindType(FStructType{EMemberKind::Struct, /* IsDynamic */ 0, /* IsSuper */ 1});
 
-template<class Struct, class CustomBinding, class Runtime>
-FMemberBindType BindMemberStruct(FOptionalInnerId& OutSchema)
+inline FMemberBindType BindInnerStruct(FOptionalInnerId& OutBindId, FMemberSpec& OutSpec, FBothStructId Both)
 {
-	// Todo: Consider auto-binding all nested structs
-	if constexpr (std::is_void_v<CustomBinding>)
-	{
-		OutSchema = FInnerId(GetStructBindId<typename Runtime::Ids, Struct>());
-	}
-	else
-	{
-		OutSchema = FInnerId(BindCustomStructOnce<Struct, CustomBinding, Runtime>());
-	}
-
+	OutBindId = FInnerId(Both.BindId);
+	OutSpec = FMemberSpec(Both.DeclId);
 	return DefaultStructBindType;
 }
 
+inline FMemberBindType BindInnerLeaf(FOptionalInnerId& OutId, FMemberSpec& OutSpec, FUnpackedLeafType Leaf, FOptionalEnumId Id)
+{
+	OutId = ToOptionalInner(Id);
+	OutSpec = FMemberSpec(Leaf, Id);
+	return FMemberBindType(Leaf);
+}
+
 template<typename Struct, class Ids>
-FMemberBindType BindInnermostType(FOptionalInnerId& OutSchema)
+FMemberBindType BindInnermostType(FOptionalInnerId& OutBindId, FMemberSpec& OutSpec)
 {
-	OutSchema = FInnerId(GetStructBindId<Ids, Struct>());
-	return DefaultStructBindType;
+	return BindInnerStruct(OutBindId, OutSpec, GetStructBothId<Ids, Struct>());
 }
 
 template<Arithmetic Type, class Ids>
-FMemberBindType BindInnermostType(FOptionalInnerId& OutSchema)
+FMemberBindType BindInnermostType(FOptionalInnerId& OutId, FMemberSpec& OutSpec)
 {
-	OutSchema = NoId;
-	return FMemberBindType(ReflectArithmetic<Type>);
+	return BindInnerLeaf(OutId, OutSpec, ReflectArithmetic<Type>, NoId);
 }
 
 template<Enumeration Enum, class Ids>
-FMemberBindType BindInnermostType(FOptionalInnerId& OutSchema)
+FMemberBindType BindInnermostType(FOptionalInnerId& OutId, FMemberSpec& OutSpec)
 {
-	OutSchema = FInnerId(GetEnumId<Ids, Enum>());
-	return FMemberBindType(ReflectEnum<Enum>);
+	return BindInnerLeaf(OutId, OutSpec, ReflectEnum<Enum>, GetEnumId<Ids, Enum>());
 }
 
 template<typename RangeBinding>
@@ -1071,15 +1096,15 @@ TConstArrayView<FRangeBinding> GetRangeBindings() requires (N > 0)
 }
 
 template<LeafType Type, class Runtime>
-FMemberBinding BindMember(uint64 Offset)
+FMemberBinding BindMember(uint64 Offset, FMemberSpec& OutSpec)
 {
 	FMemberBinding Out(Offset);
-	Out.InnermostType = BindInnermostType<Type, typename Runtime::Ids>(Out.InnermostSchema);
+	Out.InnermostType = BindInnermostType<Type, typename Runtime::Ids>(Out.InnermostSchema, OutSpec);
 	return Out;
 }
 
 template<typename Type, class Runtime>
-FMemberBinding BindMember(uint64 Offset)
+FMemberBinding BindMember(uint64 Offset, FMemberSpec& OutSpec)
 {
 	using Ids = typename Runtime::Ids;
 	using CustomBinding = typename Runtime::template CustomBindings<Type>::Type;
@@ -1087,7 +1112,8 @@ FMemberBinding BindMember(uint64 Offset)
 	FMemberBinding Out(Offset);
 	if constexpr (!std::is_void_v<CustomBinding>)
 	{
-		Out.InnermostType = BindMemberStruct<Type, CustomBinding, Runtime>(Out.InnermostSchema);
+		FBothStructId Both = BindCustomStructOnce<Type, CustomBinding, Runtime>();
+		Out.InnermostType = BindInnerStruct(Out.InnermostSchema, OutSpec, Both);
 	}
 	else
 	{
@@ -1098,11 +1124,17 @@ FMemberBinding BindMember(uint64 Offset)
 			using InnermostType = typename TInnermostType<RangeBinding, NumRangeBindings>::Type;
 
 			Out.RangeBindings = GetRangeBindings<RangeBinding, Ids, NumRangeBindings>();
-			Out.InnermostType = BindInnermostType<InnermostType, Ids>(Out.InnermostSchema);
+			Out.InnermostType = BindInnermostType<InnermostType, Ids>(Out.InnermostSchema, OutSpec);
+			
+			for (FRangeBinding Range : ReverseIterate(Out.RangeBindings))
+			{
+				OutSpec.RangeWrap(Range.GetSizeType());
+			}
 		}
 		else
 		{
-			Out.InnermostType = BindMemberStruct<Type, CustomBinding, Runtime>(Out.InnermostSchema);
+			FBothStructId Both = GetStructBothId<typename Runtime::Ids, Type>();
+			Out.InnermostType = BindInnerStruct(Out.InnermostSchema, OutSpec, Both);
 		}
 	}
 	
@@ -1112,7 +1144,7 @@ FMemberBinding BindMember(uint64 Offset)
 ////////////////////////////////////////////////////////////////////////////////////////////////
 
 template<class Ctti, class Ids>
-FEnumId DeclareNativeEnum(FDeclarations& Out, EEnumMode Mode)
+FEnumId DeclareNativeEnum(FEnumDeclarations& Out, EEnumMode Mode)
 {
 	using UnderlyingType = std::underlying_type_t<typename Ctti::Type>;
 
@@ -1124,44 +1156,42 @@ FEnumId DeclareNativeEnum(FDeclarations& Out, EEnumMode Mode)
 		Enumerator.Name = Ids::IndexName(Ctti::Enumerators[&Enumerator - Enumerators].Name);
 		Enumerator.Constant = static_cast<uint64>(static_cast<UnderlyingType>(Ctti::Enumerators[&Enumerator - Enumerators].Constant));
 	}
-	Out.DeclareEnum(Id, Type, Mode, Enumerators, EEnumAliases::Fail);
+	Out.Declare(Id, Type, Mode, Enumerators, EEnumAliases::Fail);
 
 	return Id;
 }
 
-template<class Ctti, class Ids>
-FStructId DeclareNativeStruct(FDeclarations& Out, EMemberPresence Occupancy)
+template<class Ctti, class Runtime>
+void BindNativeStruct(FSchemaBindings& Out, FBothStructId Both, EMemberPresence Occupancy)
 {
-	using Typename = TTypename<typename Ctti::Type>;
+	using Ids = Runtime::Ids;
 	using SuperType = typename Ctti::Super;
-
-	FType Type = IndexStructName<Ids, ETypename::Decl, Typename>();
-	FDeclId Id(Ids::IndexStruct(Type));
 	FOptionalDeclId SuperId;
 	if constexpr (!std::is_void_v<SuperType>)
 	{
 		SuperId = GetStructDeclId<Ids, SuperType>();
 	}
 
-	FMemberId MemberIds[Ctti::NumVars];
-	ForEachVar<Ctti>([&]<class Var>()
-	{ 
-		MemberIds[Var::Index] = Ids::IndexMember(Var::Name);
-	});
-	Out.DeclareStruct(Id, Type, 0, MemberIds, Occupancy, SuperId);
-
-	return Id;
-}
-
-template<class Ctti, class Runtime>
-void BindNativeStruct(FSchemaBindings& Out, FBindId BindId, FDeclId DeclId)
-{
-	FMemberBinding MemberBindings[Ctti::NumVars];
-	ForEachVar<Ctti>([&]<class Var>()
-	{ 
-		MemberBindings[Var::Index] = BindMember<typename Var::Type, Runtime>(Var::Offset);
-	});
-	Out.BindStruct(BindId, DeclId, MemberBindings);
+	if constexpr (Ctti::NumVars) // Temp workaround til MakeArrayView handles zero-sized arrays
+	{
+		FMemberId MemberNames[Ctti::NumVars];
+		FMemberSpec MemberTypes[Ctti::NumVars];
+		FMemberBinding MemberBindings[Ctti::NumVars];
+		ForEachVar<Ctti>([&]<class Var>()
+		{
+			MemberNames[Var::Index] = Ids::IndexMember(Var::Name);
+			FMemberSpec& OutSpec = MemberTypes[Var::Index];
+			MemberBindings[Var::Index] = BindMember<typename Var::Type, Runtime>(Var::Offset, OutSpec);
+		});
+	
+		const FStructSpec Spec = { Both.DeclId, SuperId, /* no CTTI version yet */ 0, Occupancy, MemberNames, MemberTypes };
+		Out.BindStruct(Both.BindId, MemberBindings, Spec);
+	}
+	else
+	{
+		const FStructSpec Spec = { Both.DeclId, SuperId, /* no CTTI version yet */ 0, Occupancy, {}, {} };
+		Out.BindStruct(Both.BindId, {}, Spec);
+	}
 }
 
 //////////////////////////////////////////////////////////////////////////
diff --git a/Engine/Plugins/Experimental/PlainProps/Source/Public/PlainPropsBuildSchema.h b/Engine/Plugins/Experimental/PlainProps/Source/Public/PlainPropsBuildSchema.h
--- a/Engine/Plugins/Experimental/PlainProps/Source/Public/PlainPropsBuildSchema.h
+++ b/Engine/Plugins/Experimental/PlainProps/Source/Public/PlainPropsBuildSchema.h
@@ -44,11 +44,8 @@ struct FBuiltSchemas
 class FSchemasBuilder
 {
 public:
-	using FStructDeclarations = TConstArrayView<TUniquePtr<FStructDeclaration>>;
-	using FEnumDeclarations = TConstArrayView<TUniquePtr<FEnumDeclaration>>;
 
-	PLAINPROPS_API FSchemasBuilder(const FDeclarations& Types, const FIds& InIds, const IBindIds& InBindIds, FScratchAllocator& InScratch, ESchemaFormat InFormat);
-	PLAINPROPS_API FSchemasBuilder(FStructDeclarations InStructs, FEnumDeclarations InEnums, const FIds& InIds, const IBindIds& BindIds, FScratchAllocator& Scratch, ESchemaFormat InFormat);
+	PLAINPROPS_API FSchemasBuilder(const FIds& Names, const IDeclarations& Types, FScratchAllocator& InScratch, ESchemaFormat InFormat);
 	PLAINPROPS_API ~FSchemasBuilder();
 
 	PLAINPROPS_API FEnumSchemaBuilder&			NoteEnum(FEnumId Id);
@@ -61,12 +58,10 @@ public:
 	FDebugIds									GetDebug() const { return Debug; }
 
 private:
-	FStructDeclarations							DeclaredStructs;
-	FEnumDeclarations							DeclaredEnums;
+	const IDeclarations&						Declarations;
 	TSet<FStructId>								StructIndices;
 	TSet<FEnumId>								EnumIndices;
 	const FIds&									Ids;
-	const IBindIds&								BindIds;
 	ESchemaFormat								Format;
 	TPagedArray<FStructSchemaBuilder, 4096>		Structs;		// TPagedArray for stable references
 	TPagedArray<FEnumSchemaBuilder, 4096>		Enums;			// TPagedArray for stable references
diff --git a/Engine/Plugins/Experimental/PlainProps/Source/Public/PlainPropsDeclare.h b/Engine/Plugins/Experimental/PlainProps/Source/Public/PlainPropsDeclare.h
--- a/Engine/Plugins/Experimental/PlainProps/Source/Public/PlainPropsDeclare.h
+++ b/Engine/Plugins/Experimental/PlainProps/Source/Public/PlainPropsDeclare.h
@@ -6,9 +6,13 @@
 #include "Containers/Array.h"
 #include "Containers/ArrayView.h"
 #include "Templates/UniquePtr.h"
+#include <atomic>
+
+template<typename T> class TRefCountPtr;
 
 namespace PlainProps
 {
+struct FStructSpec;
 
 struct FEnumerator
 {
@@ -32,64 +36,75 @@ struct FEnumDeclaration
 	}
 };
 
-enum class EMemberPresence : uint8 { RequireAll, AllowSparse };
-
-struct FStructDeclaration
-{
-	uint32					RefCount;
-	FDeclId					Id;				// Could be removed, might allow declaration dedup among templated types
-	FType					Type;			// Could be removed, might allow declaration dedup among templated types
-	FOptionalDeclId			Super;
-	uint16					Version;
-	uint16					NumMembers;
-	EMemberPresence			Occupancy;
-	FMemberId				MemberOrder[0];
-
-	TConstArrayView<FMemberId> GetMemberOrder() const
-	{
-		return MakeArrayView(MemberOrder, NumMembers);
-	}
-
-	inline static constexpr uint16 MaxMembers = 0xFFFF;
-};
+////////////////////////////////////////////////////////////////////////////////////////////////
 
 // Enum values are stored as integers. Aliased enum values are illegal, including composite flags.
 // Aliases can be automatically removed on declaration or detected and fail hard.
 enum class EEnumAliases { Strip, Fail };
 
-class FDeclarations
+class FEnumDeclarations
 {
 public:
-	UE_NONCOPYABLE(FDeclarations);
-	explicit FDeclarations(FDebugIds In) : Debug(In) {}
+	UE_NONCOPYABLE(FEnumDeclarations);
+	explicit FEnumDeclarations(FDebugIds In) : Debug(In) {}
 
-	PLAINPROPS_API const FEnumDeclaration&			DeclareEnum(FEnumId Id, FType Type, EEnumMode Mode, TConstArrayView<FEnumerator> Enumerators, EEnumAliases Policy);
-	// Declare struct with ref count 1 or increment it and check that previous declaration matches
-	PLAINPROPS_API const FStructDeclaration&		DeclareStruct(FDeclId Id, FType Type, uint16 Version, TConstArrayView<FMemberId> MemberOrder, EMemberPresence Occupancy, FOptionalDeclId Super = {});
-	PLAINPROPS_API const FStructDeclaration&		DeclareNumeralStruct(FDeclId Id, FType Type, TConstArrayView<FMemberId> Numerals, EMemberPresence Occupancy);
-	
-	void											DropEnum(FEnumId Id)			{ Check(Id); DeclaredEnums[Id.Idx].Reset(); }
-	PLAINPROPS_API void								DropStructRef(FDeclId Id);
-
-	const FEnumDeclaration&							Get(FEnumId Id) const			{ Check(Id); return *DeclaredEnums[Id.Idx]; }
-	const FStructDeclaration&						Get(FDeclId Id) const			{ Check(Id); return *DeclaredStructs[Id.Idx]; }
-	PLAINPROPS_API const FStructDeclaration*		Find(FDeclId Id) const;
-	
-	TConstArrayView<TUniquePtr<FEnumDeclaration>>	GetEnums() const				{ return DeclaredEnums; }
-	TConstArrayView<TUniquePtr<FStructDeclaration>>	GetStructs() const				{ return DeclaredStructs; }
-	FDebugIds										GetDebug() const				{ return Debug; }
+	PLAINPROPS_API const FEnumDeclaration&	Declare(FEnumId Id, FType Type, EEnumMode Mode, TConstArrayView<FEnumerator> Enumerators, EEnumAliases Policy);
+	void									Drop(FEnumId Id)				{ Check(Id); Declarations[Id.Idx].Reset(); }
+	const FEnumDeclaration&					Get(FEnumId Id) const			{ Check(Id); return *Declarations[Id.Idx]; }
+	FDebugIds								GetDebug() const				{ return Debug; }
 
 protected:
-	TArray<TUniquePtr<FEnumDeclaration>>			DeclaredEnums;
-	TArray<TUniquePtr<FStructDeclaration>>			DeclaredStructs;
-	FDebugIds										Debug;
+	TArray<TUniquePtr<FEnumDeclaration>>	Declarations;
+	FDebugIds								Debug;
 
 #if DO_CHECK
-	PLAINPROPS_API void								Check(FEnumId Id) const;
-	PLAINPROPS_API void								Check(FDeclId Id) const;
+	PLAINPROPS_API void						Check(FEnumId Id) const;
 #else
-	void											Check(...) const {}
+	void									Check(FEnumId) const {}
 #endif
 };
 
-} // namespace PlainProps
+////////////////////////////////////////////////////////////////////////////////////////////////
+
+enum class EMemberPresence : uint8 { RequireAll, AllowSparse };
+
+struct FStructDeclaration
+{
+	mutable std::atomic<int32>		RefCount;
+	FDeclId							Id;				// Could be removed, might allow declaration dedup among templated types
+	FOptionalDeclId					Super;
+	uint16							Version;
+	uint16							NumMembers;
+	uint16							NumInnerIds;
+	uint16							NumInnerRanges;
+	EMemberPresence					Occupancy;
+	FMemberId						MemberNames[0];
+
+	const FInnerId*					GetInnerIds() const			{ return AlignPtr<FInnerId>(MemberNames + NumMembers); }
+	const FMemberType*				GetTypes() const			{ return AlignPtr<FMemberType>(GetInnerIds() + NumInnerIds); }
+	const FMemberType*				GetInnerRangeTypes() const	{ return GetTypes() + NumMembers; }
+	TConstArrayView<FMemberId>		GetMemberOrder() const		{ return MakeArrayView(MemberNames, NumMembers); }
+	uint32							CalculateSize() const;
+	
+	void							AddRef() const				{ RefCount.fetch_add(1, std::memory_order_relaxed);	}
+	PLAINPROPS_API bool				Release() const;
+
+	inline static constexpr uint16 MaxMembers = 0xFFFF;
+};
+
+////////////////////////////////////////////////////////////////////////////////////////////////
+
+using FStructDeclarationPtr = ::TRefCountPtr<const FStructDeclaration>;
+
+PLAINPROPS_API FStructDeclarationPtr Declare(FStructSpec Spec);
+
+////////////////////////////////////////////////////////////////////////////////////////////////
+
+struct IDeclarations
+{
+	virtual FDeclId						Lower(FBindId BindId) const = 0;
+	virtual const FEnumDeclaration*		Find(FEnumId Id) const = 0;
+	virtual const FStructDeclaration*	Find(FStructId Id) const = 0;
+};
+
+} // namespace PlainProps
\ No newline at end of file
diff --git a/Engine/Plugins/Experimental/PlainProps/Source/Public/PlainPropsLoad.h b/Engine/Plugins/Experimental/PlainProps/Source/Public/PlainPropsLoad.h
--- a/Engine/Plugins/Experimental/PlainProps/Source/Public/PlainPropsLoad.h
+++ b/Engine/Plugins/Experimental/PlainProps/Source/Public/PlainPropsLoad.h
@@ -11,18 +11,17 @@ namespace PlainProps
 {
 	
 class FCustomBindings;
-class FDeclarations;
 class FSchemaBindings;
 struct FLoadBatch;
 
-struct FLoadBatchDeleter 
+struct FLoadBatchDeleter
 {
 	PLAINPROPS_API void operator()(FLoadBatch* Ptr) const;
 };
 using FLoadBatchPtr = TUniquePtr<FLoadBatch, FLoadBatchDeleter>;
 
 
-[[nodiscard]] PLAINPROPS_API FLoadBatchPtr	CreateLoadPlans(FSchemaBatchId BatchId, const FDeclarations& Declarations, const FCustomBindings& Customs, const FSchemaBindings& Schemas, TConstArrayView<FStructId> RuntimeIds, ESchemaFormat Format);
+[[nodiscard]] PLAINPROPS_API FLoadBatchPtr	CreateLoadPlans(FSchemaBatchId BatchId,const FCustomBindings& Customs, const FSchemaBindings& Schemas, TConstArrayView<FStructId> RuntimeIds, ESchemaFormat Format);
 PLAINPROPS_API void							LoadStruct(void* Dst, FByteReader Src, FStructSchemaId LoadId, const FLoadBatch& Batch);
 PLAINPROPS_API void							ConstructAndLoadStruct(void* Dst, FByteReader Src, FStructSchemaId LoadId, const FLoadBatch& Batch);
 
diff --git a/Engine/Plugins/Experimental/PlainProps/Source/Public/PlainPropsSaveMember.h b/Engine/Plugins/Experimental/PlainProps/Source/Public/PlainPropsSaveMember.h
--- a/Engine/Plugins/Experimental/PlainProps/Source/Public/PlainPropsSaveMember.h
+++ b/Engine/Plugins/Experimental/PlainProps/Source/Public/PlainPropsSaveMember.h
@@ -29,20 +29,22 @@ FORCEINLINE bool Track(bool bDiff, FDiffContext& Ctx, FMemberBindType Type, FMem
 
 // TMemberSerializer helper
 template<typename T>
-struct TStructSerializer
+struct TStructSerializer : FBothStructId
 {
 	static constexpr EMemberKind Kind = EMemberKind::Struct;
 
 	template<class Ids>
-	TStructSerializer(TCustomInit<Ids> Init, TConstArrayView<FMemberId> Names, FBothStructId Both = GetStructBothId<Ids, T>())
-	: Id(Both.BindId)
+	TStructSerializer(TCustomInit<Ids>& Init, TConstArrayView<FMemberId> Names, FBothStructId Both = GetStructBothId<Ids, T>())
+	: FBothStructId(Both)
 	{
 		Init.RegisterInnerStruct(Both, Names);
 	}
 
+	FMemberSpec SpecMember() const { return FMemberSpec(DeclId); }
+
 	FTypedValue SaveMember(const T& Value, const FSaveContext& Ctx) const
 	{
-		FMemberSchema Schema = { DefaultStructType, DefaultStructType, 0, FInnerId(Id), nullptr};
+		FMemberSchema Schema = { DefaultStructType, DefaultStructType, 0, FInnerId(BindId), nullptr};
 		return { Schema, {.Struct = SaveItem(Value, Ctx)} };
 	}
 
@@ -67,7 +69,7 @@ struct TStructSerializer
 	template<class ContextType>
 	bool DiffMember(const T& A, const T& B, FMemberId Name, ContextType& Ctx) const
 	{
-		return Track(DiffStructs(&A, &B, Id, Ctx), Ctx, DefaultStructBindType, Name, {.Struct = Id}, &A, &B);
+		return Track(DiffStructs(&A, &B, BindId, Ctx), Ctx, DefaultStructBindType, Name, {.Struct = BindId}, &A, &B);
 	}
 
 	////////// TRangeSaver API //////////
@@ -76,15 +78,13 @@ struct TStructSerializer
 
 	FBuiltStruct* SaveItem(const T& Value, const FSaveContext& Ctx) const
 	{
-		return SaveStruct(&Value, Id, Ctx);
+		return SaveStruct(&Value, BindId, Ctx);
 	}
 
 	FMemberSchema MakeMemberRangeSchema(ERangeSizeType MaxSize) const
 	{
-		return { FMemberType(MaxSize), DefaultStructType, 1, FInnerId(Id), nullptr };
+		return { FMemberType(MaxSize), DefaultStructType, 1, FInnerId(BindId), nullptr };
 	}
-
-	const FBindId Id;
 };
 
 //////////////////////////////////////////////////////////////////////////
@@ -124,12 +124,13 @@ struct TRangeSerializer
 	static constexpr uint16			NumTypes = NumRanges + 1;
 
 	FRangeBinding const* const		Bindings = nullptr;
-	FOptionalInnerId				InnermostSchema;
+	FOptionalInnerId				InnermostId;
+	FOptionalInnerId				InnermostSpecId;
 	FMemberBindTypeInitializer		BindTypes[NumTypes];
 	FMemberType						Types[NumTypes];
 
 	template<class Ids>
-	TRangeSerializer(TCustomInit<Ids> Init, TConstArrayView<FMemberId> Names)
+	TRangeSerializer(TCustomInit<Ids>& Init, TConstArrayView<FMemberId> Names)
 	: Bindings(GetRangeBindings<RangeBinding, Ids, NumRanges>().GetData())
 	{
 		for (uint16 Idx = 0; Idx < NumRanges; ++Idx)
@@ -141,16 +142,24 @@ struct TRangeSerializer
 
 		using InnermostType = typename TInnermostType<RangeBinding, NumRanges>::Type;
 		Types[NumRanges] = ReflectInnermostType<InnermostType>();
-		BindTypes[NumRanges].Value = BindInnermostType<InnermostType, Ids>(/* out */ InnermostSchema);
 		
+		FMemberSpec Spec;
+		BindTypes[NumRanges].Value = BindInnermostType<InnermostType, Ids>(/* out */ InnermostId, /* out */ Spec);
+		InnermostSpecId = InnermostId;
 
 		if constexpr (!LeafType<InnermostType>)
 		{			
-			FBothStructId Id = {InnermostSchema.Get().AsStructBindId(), GetStructDeclId<Ids, InnermostType>()};
-			Init.RegisterInnerStruct(Id, Names);
+			FBothStructId Both = {InnermostId.Get().AsStructBindId(), GetStructDeclId<Ids, InnermostType>() };
+			InnermostSpecId = FInnerId(Both.DeclId);
+			Init.RegisterInnerStruct(Both, Names);
 		}
 	}
 
+	FMemberSpec SpecMember() const
+	{
+		return FMemberSpec(Types, InnermostSpecId);
+	}
+
 	void LoadMember(T& Dst, /* in-out */ FMemberLoader& Src) const
 	{
 		LoadRange(&Dst, Src.GrabRange(), MakeArrayView(Bindings, NumRanges));
@@ -176,7 +185,7 @@ struct TRangeSerializer
 		}
 		else if constexpr (Kind == EMemberKind::Struct)
 		{
-			return DiffRanges(A, B, Binding, InnermostSchema.Get().AsStructBindId(), Ctx);
+			return DiffRanges(A, B, Binding, InnermostId.Get().AsStructBindId(), Ctx);
 		}
 		else
 		{
@@ -188,11 +197,11 @@ struct TRangeSerializer
 	{
 		if constexpr (Kind == EMemberKind::Leaf)
 		{
-			return { .Leaf = ToOptionalEnum(InnermostSchema) };
+			return { .Leaf = ToOptionalEnum(InnermostId) };
 		}
 		else if constexpr (Kind == EMemberKind::Struct)
 		{
-			return { .Struct = InnermostSchema.Get().AsStructBindId() };
+			return { .Struct = InnermostId.Get().AsStructBindId() };
 		}
 		else
 		{
@@ -210,12 +219,12 @@ struct TRangeSerializer
 		
 	FMemberSchema MakeMemberSchema() const
 	{
-		return { Types[0], Types[1], NumRanges, InnermostSchema, NumRanges > 1 ? Types + 1 : nullptr };
+		return { Types[0], Types[1], NumRanges, InnermostId, NumRanges > 1 ? Types + 1 : nullptr };
 	}
 
 	FRangeMemberBinding MakeInnerRangeBinding() const
 	{
-		return { &BindTypes[1].Value, Bindings, NumRanges, InnermostSchema, /* offset */ 0 };
+		return { &BindTypes[1].Value, Bindings, NumRanges, InnermostId, /* offset */ 0 };
 	}
 
 	////////// TRangeSaver API //////////
@@ -224,7 +233,7 @@ struct TRangeSerializer
 
 	FMemberSchema MakeMemberRangeSchema(ERangeSizeType MaxSize) const
 	{
-		return { FMemberType(MaxSize), Types[0], NumTypes, InnermostSchema, Types };
+		return { FMemberType(MaxSize), Types[0], NumTypes, InnermostId, Types };
 	}
 
 	FBuiltRange* SaveItem(const T& Item, const FSaveContext& Ctx) const
@@ -235,6 +244,7 @@ struct TRangeSerializer
 		}
 		else
 		{
+			FRangeMemberBinding MemberBinding = { &BindTypes[1].Value, Bindings, NumRanges, InnermostId, /* offset */ 0 };
 			return SaveRange(&Item, MakeInnerRangeBinding(), Ctx);
 		}
 	}
@@ -276,9 +286,14 @@ struct TArithmeticSerializer : TLeafSerializer<T>
 	static constexpr FMemberBindType MemberBindType{ReflectArithmetic<T>};
 
 	template<class Ids>
-	TArithmeticSerializer(TCustomInit<Ids>, TConstArrayView<FMemberId>)
+	TArithmeticSerializer(TCustomInit<Ids>&, TConstArrayView<FMemberId>)
 	{}
 
+	static FMemberSpec SpecMember()
+	{
+		return Specify<T>();
+	}
+
 	static FTypedValue SaveMember(T Value, const FSaveContext&)
 	{
 		return TLeafSerializer<T>::SaveLeaf(MemberType, NoId, Value);
@@ -308,10 +323,15 @@ struct TEnumSerializer : TLeafSerializer<T>
 	const FEnumId Id;
 
 	template<class Ids>
-	TEnumSerializer(TCustomInit<Ids>, TConstArrayView<FMemberId>)
+	TEnumSerializer(TCustomInit<Ids>&, TConstArrayView<FMemberId>)
 	: Id(GetEnumId<Ids, T>())
 	{}
 
+	FMemberSpec SpecMember() const
+	{
+		return Specify<T>(Id);
+	}
+	
 	FTypedValue SaveMember(T Value, const FSaveContext& Ctx) const
 	{
 		return TLeafSerializer<T>::SaveLeaf(MemberType, Id, Value);
diff --git a/Engine/Plugins/Experimental/PlainProps/Source/Public/PlainPropsSpecify.h b/Engine/Plugins/Experimental/PlainProps/Source/Public/PlainPropsSpecify.h
new file mode 100644
--- /dev/null
+++ b/Engine/Plugins/Experimental/PlainProps/Source/Public/PlainPropsSpecify.h
@@ -0,0 +1,123 @@
+// Copyright Epic Games, Inc. All Rights Reserved.
+
+#pragma once
+
+#include "Containers/ArrayView.h"
+#include "PlainPropsDeclare.h"
+#include "PlainPropsTypes.h"
+
+namespace PlainProps 
+{
+
+enum ESpecBool { SpecBool };
+enum ESpecIntS { SpecS8, SpecS16, SpecS32, SpecS64 };			// Values match ELeafWidth
+enum ESpecIntU { SpecU8, SpecU16, SpecU32, SpecU64 };			// Values match ELeafWidth
+enum ESpecHex { SpecHex8, SpecHex16, SpecHex32, SpecHex64 };	// Values match ELeafWidth
+enum ESpecUtf { SpecUtf8, SpecUtf16, SpecUtf32 };				// Values match ELeafWidth
+enum ESpecFloat { SpecF32 = 2, SpecF64 = 3 };					// Values match ELeafWidth
+enum ESpecDynamicStruct { SpecDynamicStruct };
+
+struct FMemberSpecView
+{
+	FMemberType						InnermostType;
+	FOptionalInnerId				InnermostId;
+	TConstArrayView<ERangeSizeType>	Ranges;
+};
+
+// Specifies the type of a member before declaring it
+class FMemberSpec
+{
+public:
+	FMemberSpec() : bSpecified(false) {}
+
+	// @pre !Type.IsRange()
+	FMemberSpec(FMemberType Type, FOptionalInnerId Id) : InnermostType(Type), InnermostId(Id) { check(!Type.IsRange());	}
+
+	FMemberSpec(FUnpackedLeafType Leaf, FOptionalEnumId Id) : FMemberSpec(FMemberType(Leaf.Pack()), FOptionalInnerId(Id)) {}
+	FMemberSpec(ELeafType Type, ELeafWidth Width, FOptionalEnumId Id) : FMemberSpec({Type, Width}, Id) {}
+	FMemberSpec(ESpecBool)		: FMemberSpec(ELeafType::Bool,		ELeafWidth::B8,	NoId) {}
+	FMemberSpec(ESpecIntS In)	: FMemberSpec(ELeafType::IntS,		static_cast<ELeafWidth>(In), NoId) {}
+	FMemberSpec(ESpecIntU In)	: FMemberSpec(ELeafType::IntU,		static_cast<ELeafWidth>(In), NoId) {}
+	FMemberSpec(ESpecHex In)	: FMemberSpec(ELeafType::Hex,		static_cast<ELeafWidth>(In), NoId) {}
+	FMemberSpec(ESpecUtf In)	: FMemberSpec(ELeafType::Unicode,	static_cast<ELeafWidth>(In), NoId) {}
+	FMemberSpec(ESpecFloat In)	: FMemberSpec(ELeafType::Float,		static_cast<ELeafWidth>(In), NoId) {}
+	
+	FMemberSpec(FDeclId Id)			: FMemberSpec(FStructType{EMemberKind::Struct, /* IsDynamic */ 0, /* IsSuper */ 0}, FInnerId(Id)) {}
+	FMemberSpec(ESpecDynamicStruct) : FMemberSpec(FStructType{EMemberKind::Struct, /* IsDynamic */ 1, /* IsSuper */ 0}, NoId) {}
+
+	PLAINPROPS_API FMemberSpec(ERangeSizeType MaxSize, FMemberSpec Inner);
+	PLAINPROPS_API FMemberSpec(TConstArrayView<FMemberType> Members, FOptionalInnerId InInnermostId);
+	PLAINPROPS_API FMemberSpec(FMemberType Type, TConstArrayView<FMemberType> InnerRangeTypes, FOptionalInnerId InInnermostId);
+	
+	void RangeWrap(ERangeSizeType MaxSize)	{ Ranges.Emplace(MaxSize); }
+	void RangeWrap(FMemberType Range)		{ Ranges.Emplace(Range.AsRange().MaxSize); }
+
+	operator FMemberSpecView() const UE_LIFETIMEBOUND;
+	
+private:
+	using RangeSizeArray = TArray<ERangeSizeType, TInlineAllocator<8>>;
+	
+	bool							bSpecified = true;						
+	FMemberType						InnermostType;
+	FOptionalInnerId				InnermostId;
+	RangeSizeArray					Ranges;
+
+	FMemberSpec(FStructType Type, FOptionalInnerId Id) : InnermostType(Type), InnermostId(Id) {}
+};
+
+inline FMemberSpec::operator FMemberSpecView() const
+{
+	check(bSpecified);
+	return { InnermostType, InnermostId, Ranges };
+}
+
+////////////////////////////////////////////////////////////////////////////////////////////////
+
+template<Arithmetic T>
+inline FMemberSpec Specify()
+{
+	return FMemberSpec(ReflectArithmetic<T>, NoId);
+}
+
+template<Enumeration T>
+inline FMemberSpec Specify(FEnumId Id)
+{
+	return FMemberSpec(ReflectEnum<T>, Id);
+}
+
+////////////////////////////////////////////////////////////////////////////////////////////////
+
+// Struct declaration input
+struct FStructSpec
+{
+	FDeclId									Id;
+	FOptionalDeclId							Super;
+	uint16									Version = 0;
+	EMemberPresence							Occupancy;
+	TConstArrayView<FMemberId>				MemberNames;
+	TConstArrayView<FMemberSpec>			MemberTypes;
+};
+
+////////////////////////////////////////////////////////////////////////////////////////////////
+
+// Helps specify single member structs
+struct FSoleMemberSpec
+{
+	FSoleMemberSpec(FDeclId InId, FMemberId Name, FMemberSpec Spec, FOptionalDeclId InSuper = NoId, uint16 Ver = 0)
+	: Id(InId), Super(InSuper), Version(Ver), MemberName(Name), MemberType(Spec)
+	{}
+
+	FDeclId									Id;
+	FOptionalDeclId							Super;
+	uint16									Version;
+	FMemberId								MemberName;
+	FMemberSpec								MemberType;
+
+	operator FStructSpec() const UE_LIFETIMEBOUND
+	{
+		return {Id, Super, Version, EMemberPresence::RequireAll, MakeArrayView(&MemberName, 1), MakeArrayView(&MemberType, 1) };
+	}
+};
+
+
+} // namespace PlainProps
\ No newline at end of file
diff --git a/Engine/Plugins/Experimental/PlainProps/Source/Public/PlainPropsTypes.h b/Engine/Plugins/Experimental/PlainProps/Source/Public/PlainPropsTypes.h
--- a/Engine/Plugins/Experimental/PlainProps/Source/Public/PlainPropsTypes.h
+++ b/Engine/Plugins/Experimental/PlainProps/Source/Public/PlainPropsTypes.h
@@ -157,13 +157,6 @@ struct FBindId : FStructId {};
 inline FDeclId LowerCast(FBindId Id) { return static_cast<FDeclId>(FStructId(Id)); }
 inline FBindId UpCast(FDeclId Id) { return static_cast<FBindId>(FStructId(Id)); }
 
-// Interface for type-erasing / lowering FBindId to FDeclId
-struct IBindIds
-{
-	virtual ~IBindIds() = default;
-	virtual FDeclId Lower(FBindId BindId) const = 0;
-};
-
 // Either a runtime FStructId or FEnumId of a member
 struct FInnerId
 {
@@ -360,6 +353,8 @@ using FOptionalConcreteTypenameId = TOptionalId<FConcreteTypenameId>;
 template<class IdType>
 inline constexpr TOptionalId<IdType> ToOptional(IdType Id) { return Id; }
 
+inline FOptionalInnerId			ToOptionalInner(FOptionalEnumId In)		{ return static_cast<const FOptionalInnerId&>(static_cast<const FOptionalId&>(In)); }
+inline FOptionalInnerId			ToOptionalInner(FOptionalStructId In)	{ return static_cast<const FOptionalInnerId&>(static_cast<const FOptionalId&>(In)); }
 inline FOptionalEnumId			ToOptionalEnum(FOptionalInnerId In)		{ return static_cast<const FOptionalEnumId&>(static_cast<const FOptionalId&>(In)); }
 inline FOptionalStructId		ToOptionalStruct(FOptionalInnerId In)	{ return static_cast<const FOptionalStructId&>(static_cast<const FOptionalId&>(In)); }
 inline FOptionalStructId		ToOptionalStruct(FOptionalDeclId In)	{ return static_cast<const FOptionalStructId&>(static_cast<const FOptionalId&>(In)); }
@@ -380,7 +375,6 @@ struct FNestedScope
 	friend uint32 GetTypeHash(FNestedScope Scope) { return HashCombineFast(Scope.Outer.AsInt(), Scope.Inner.Name.Idx); };
 };
 
-
 struct FParameterIndexRange : FBaseTypenameId { using FBaseTypenameId::FBaseTypenameId; };
 
 // Name-resolved FParametricTypeId
@@ -466,7 +460,6 @@ template<> inline constexpr FUnpackedLeafType ReflectArithmetic<char8_t>	= { ELe
 template<> inline constexpr FUnpackedLeafType ReflectArithmetic<char16_t>	= { ELeafType::Unicode,	ELeafWidth::B16 };
 template<> inline constexpr FUnpackedLeafType ReflectArithmetic<char32_t>	= { ELeafType::Unicode,	ELeafWidth::B32 };
 
-
 template<Arithmetic T>
 inline constexpr FUnpackedLeafType MakeLeafType() { return ReflectArithmetic<T>; }
 
@@ -561,4 +554,4 @@ private:
 	const FIds& Ids;
 };
 
-} // namespace PlainProps
+} // namespace PlainProps
\ No newline at end of file
diff --git a/Engine/Plugins/Experimental/PlainProps/Source/Public/PlainPropsUeCoreBindings.h b/Engine/Plugins/Experimental/PlainProps/Source/Public/PlainPropsUeCoreBindings.h
--- a/Engine/Plugins/Experimental/PlainProps/Source/Public/PlainPropsUeCoreBindings.h
+++ b/Engine/Plugins/Experimental/PlainProps/Source/Public/PlainPropsUeCoreBindings.h
@@ -568,19 +568,21 @@ struct TPairBinding : ICustomBinding
 	using Type = TPair<K,V>;
 
 	template<class Ids>
-	TPairBinding(TCustomInit<Ids> Init)
+	TPairBinding(TCustomSpecifier<Ids, 2>& Spec)
 	: MemberIds{Ids::IndexMember("Key"), Ids::IndexMember("Value")}
-	, Key(Init, { MemberIds[0] })
-	, Value(Init, { MemberIds[1] })
-	{}
+	, Key(Spec, { MemberIds[0] })
+	, Value(Spec, { MemberIds[1] })
+	{
+		Spec.SetMembers(Key.SpecMember(), Value.SpecMember());
+	}
 
-	void Save(FMemberBuilder& Dst, const TPair<K,V>& Src, const Type* Default, const FSaveContext& Ctx) const
+	inline void Save(FMemberBuilder& Dst, const TPair<K,V>& Src, const Type* Default, const FSaveContext& Ctx) const
 	{
 		Dst.Add(MemberIds[0], Key.SaveMember(Src.Key, Ctx));
 		Dst.Add(MemberIds[1], Value.SaveMember(Src.Value, Ctx));
 	}
 		
-	void Load(TPair<K,V>& Dst, FStructLoadView Src, ECustomLoadMethod Method) const
+	inline void Load(TPair<K,V>& Dst, FStructLoadView Src, ECustomLoadMethod Method) const
 	{
 		FMemberLoader Members(Src);
 		check(Members.PeekName() == MemberIds[0]);
@@ -626,17 +628,20 @@ struct FTransformBinding : ICustomBinding
 {
 	using Type = FTransform;
 	enum class EMember : uint8 { Translate, Rotate, Scale };
-	const FMemberId		MemberIds[3];
-	const FBindId		VectorId;
-	const FBindId		QuatId;
+
+	const FMemberId					MemberIds[3];
+	const FDualStructId				VectorId;
+	const FDualStructId				QuatId;
 
 	template<class Ids>
-	FTransformBinding(TCustomInit<Ids>)
+	FTransformBinding(TCustomSpecifier<Ids, 3>& Spec)
 	: MemberIds{Ids::IndexMember("Translate"), Ids::IndexMember("Rotate"), Ids::IndexMember("Scale")}
 	, VectorId(GetStructBindId<Ids, FVector>())
 	, QuatId(GetStructBindId<Ids, FQuat>())
-	{}
-
+	{
+		Spec.SetMembers(VectorId, QuatId, VectorId);
+	}
+	
 	PLAINPROPS_API void	Save(FMemberBuilder& Dst, const FTransform& Src, const FTransform* Default, const FSaveContext& Context) const;
 	PLAINPROPS_API void	Load(FTransform& Dst, FStructLoadView Src, ECustomLoadMethod Method) const;
 	PLAINPROPS_API bool Diff(const FTransform& A, const FTransform& B, FDiffContext& Ctx) const;
@@ -649,11 +654,13 @@ struct FGuidBinding : ICustomBinding
 {
 	using Type = FGuid;
 	const FMemberId MemberIds[4];
-
+	
 	template<class Ids>
-	FGuidBinding(TCustomInit<Ids>)
+	FGuidBinding(TCustomSpecifier<Ids, 4>& Spec)
 	: MemberIds{Ids::IndexMember("A"), Ids::IndexMember("B"), Ids::IndexMember("C"), Ids::IndexMember("D")}
-	{}
+	{
+		Spec.FillMembers(SpecHex32);
+	}
 
 	PLAINPROPS_API void	Save(FMemberBuilder& Dst, const FGuid& Src, const FGuid* Default, const FSaveContext&) const;
 	PLAINPROPS_API void	Load(FGuid& Dst, FStructLoadView Src, ECustomLoadMethod Method) const;
@@ -668,10 +675,12 @@ struct FColorBinding : ICustomBinding
 	const FMemberId MemberIds[4];
 
 	template<class Ids>
-	FColorBinding(TCustomInit<Ids>)
+	FColorBinding(TCustomSpecifier<Ids, 4>& Spec)
 	: MemberIds{Ids::IndexMember("B"), Ids::IndexMember("G"), Ids::IndexMember("R"), Ids::IndexMember("A")}
-	{}
-
+	{
+		Spec.FillMembers(Specify<decltype(FColor::R)>());
+	}
+	
 	PLAINPROPS_API void	Save(FMemberBuilder& Dst, const FColor& Src, const FColor* Default, const FSaveContext&) const;
 	PLAINPROPS_API void	Load(FColor& Dst, FStructLoadView Src, ECustomLoadMethod Method) const;
 	inline static  bool Diff(FColor A, FColor B, const FBindContext&) { return A != B; }
@@ -685,10 +694,12 @@ struct FLinearColorBinding : ICustomBinding
 	const FMemberId MemberIds[4];
 
 	template<class Ids>
-	FLinearColorBinding(TCustomInit<Ids>)
+	FLinearColorBinding(TCustomSpecifier<Ids, 4>& Spec)
 	: MemberIds{Ids::IndexMember("R"), Ids::IndexMember("G"), Ids::IndexMember("B"), Ids::IndexMember("A")}
-	{}
-
+	{
+		Spec.FillMembers(Specify<decltype(FLinearColor::R)>());
+	}
+	
 	PLAINPROPS_API void	Save(FMemberBuilder& Dst, const FLinearColor& Src, const FLinearColor* Default, const FSaveContext&) const;
 	PLAINPROPS_API void	Load(FLinearColor& Dst, FStructLoadView Src, ECustomLoadMethod Method) const;
 	inline static  bool Diff(FLinearColor A, FLinearColor B, const FBindContext&) { return A != B; }
@@ -727,7 +738,7 @@ struct TBaseDeltaBinding : ICustomBinding, FBaseDeltaBinding
 	const TMemberSerializer<ElemType> Elems;
 
 	template<class Ids>	
-	TBaseDeltaBinding(TCustomInit<Ids> Init)
+	TBaseDeltaBinding(TCustomInit<Ids>& Init)
 	: FBaseDeltaBinding(FBaseDeltaBinding::Cache<Ids>())
 	, Elems(Init, bIsSet ? MakeArrayView(MemberIds) : MakeArrayView({AddId()}))
 	{}
@@ -969,10 +980,12 @@ struct TSetDeltaBinding : TBaseDeltaBinding<TSet<T, KeyFuncs, SetAllocator>, T>
 	};
 		
 	template<class Ids>	
-	TSetDeltaBinding(TCustomInit<Ids> Init)
-	: Super(Init)
-	{}
-	
+	TSetDeltaBinding(TCustomSpecifier<Ids, 2>& Spec)
+	: Super(Spec)
+	{
+		Spec.FillMembers(FMemberSpec(Super::MaxSize, Super::Elems.SpecMember()));
+	}
+
 	inline void Save(FMemberBuilder& Dst, const SetType& Src, const SetType* Default, const FSaveContext& Ctx) const
 	{
 		Super::SaveDelta(Dst, Src, Default, Ctx, /* Keys */ Super::Elems);
@@ -1001,11 +1014,14 @@ struct TMapDeltaBinding : TBaseDeltaBinding<TMap<K, V, SetAllocator, KeyFuncs>,
 	const TMemberSerializer<K> Keys;
 	
 	template<class Ids>	
-	TMapDeltaBinding(TCustomInit<Ids> Init)
-	: Super(Init)
-	, Keys(Init, {FBaseDeltaBinding::DelId()})
-	{}
-	
+	TMapDeltaBinding(TCustomSpecifier<Ids, 2>& Spec)
+	: Super(Spec)
+	, Keys(Spec, {FBaseDeltaBinding::DelId()})
+	{
+		Spec.Members[0] = FMemberSpec(Super::MaxSize, Keys.SpecMember());
+		Spec.Members[1] = FMemberSpec(Super::MaxSize, Super::Elems.SpecMember());
+	}
+		
 	inline void Save(FMemberBuilder& Dst, const MapType& Src, const MapType* Default, const FSaveContext& Ctx) const
 	{
 		Super::SaveDelta(Dst, Src, Default, Ctx, Keys);
@@ -1017,17 +1033,6 @@ struct TMapDeltaBinding : TBaseDeltaBinding<TMap<K, V, SetAllocator, KeyFuncs>,
 	}
 };
 
-struct FEmptyVariantStateBinding : ICustomBinding
-{
-	using Type = FEmptyVariantState;
-	static constexpr std::initializer_list<FMemberId> MemberIds = {};
-
-	template<class Ids> FEmptyVariantStateBinding(TCustomInit<Ids>) {}
-	void Save(FMemberBuilder& Dst, const FEmptyVariantState& Src, const FEmptyVariantState*, const FSaveContext&) const { }
-	void Load(FEmptyVariantState& Dst, FStructLoadView Src, ECustomLoadMethod Method) const { }
-	static bool Diff(const FEmptyVariantState& A, const FEmptyVariantState& B, const FBindContext&) { return false; }
-};
-
 template <typename T, typename... Ts>
 struct TVariantMemberSerializer
 {
@@ -1070,20 +1075,23 @@ struct TVariantBinding : ICustomBinding
 {
 	using Type = TVariant<Ts...>;
 	using TSerializers = std::tuple<TMemberSerializer<Ts>...>;
+	static constexpr uint32 N = sizeof...(Ts);
 
-	const FMemberId MemberIds[sizeof...(Ts)];
+	const FMemberId MemberIds[N];
 	const TSerializers Serializers;
-	static constexpr TVariantMemberFunctions<Ts...> Functions[sizeof...(Ts)] = {{
+	static constexpr TVariantMemberFunctions<Ts...> Functions[N] = {{
 		&TVariantMemberSerializer<Ts,Ts...>::SaveMember,
 		&TVariantMemberSerializer<Ts,Ts...>::LoadMember,
 		&TVariantMemberSerializer<Ts,Ts...>::DiffMember,
 		&TVariantMemberSerializer<Ts,Ts...>::DiffMember}...};
 
 	template<class Ids>
-	TVariantBinding(TCustomInit<Ids> Init)
+	TVariantBinding(TCustomSpecifier<Ids, N>& Spec)
 	: MemberIds{Ids::IndexNumeral(::UE::Core::Private::TParameterPackTypeIndex<Ts, Ts...>::Value)...}
-	, Serializers{{Init, MakeArrayView({MemberIds[::UE::Core::Private::TParameterPackTypeIndex<Ts, Ts...>::Value]})}...}
-	{ }
+	, Serializers{{Spec, MakeArrayView({MemberIds[::UE::Core::Private::TParameterPackTypeIndex<Ts, Ts...>::Value]})}...}
+	{
+		Spec.SetMembers(std::get<TMemberSerializer<Ts>>(Serializers).SpecMember()...);
+	}
 
 	void Save(FMemberBuilder& Dst, const Type& Src, const Type* Default, const FSaveContext& Ctx) const
 	{
@@ -1140,10 +1148,12 @@ template <> struct TTypename<FTransform>			{ inline static constexpr std::string
 template <> struct TTypename<FGuid>					{ inline static constexpr std::string_view DeclName = "Guid"; };
 template <> struct TTypename<FColor>				{ inline static constexpr std::string_view DeclName = "Color"; };
 template <> struct TTypename<FLinearColor>			{ inline static constexpr std::string_view DeclName = "LinearColor"; };
+template <> struct TTypename<FEmptyVariantState>	{ inline static constexpr std::string_view DeclName = "EmptyVariant"; };
 template <> struct TTypename<FString>				{ inline static constexpr std::string_view RangeBindName = "String"; };
 template <> struct TTypename<FUtf8String>			{ inline static constexpr std::string_view RangeBindName = "Utf8String"; };
 template <> struct TTypename<FAnsiString>			{ inline static constexpr std::string_view RangeBindName = "AnsiString"; };
 
+
 template <typename K, typename V>
 struct TTypename<TPair<K,V>>
 {
@@ -1270,8 +1280,6 @@ struct TCustomBind<TPair<K, V>>
 	using Type = UE::TPairBinding<K, V>;
 };
 
-template<> struct TCustomBind<FEmptyVariantState> { using Type = UE::FEmptyVariantStateBinding; };
-
 template<typename... Ts>
 struct TCustomBind<TVariant<Ts...>>
 {
diff --git a/Engine/Plugins/Experimental/PlainProps/Source/Public/PlainPropsWrite.h b/Engine/Plugins/Experimental/PlainProps/Source/Public/PlainPropsWrite.h
--- a/Engine/Plugins/Experimental/PlainProps/Source/Public/PlainPropsWrite.h
+++ b/Engine/Plugins/Experimental/PlainProps/Source/Public/PlainPropsWrite.h
@@ -16,12 +16,12 @@ class FIdIndexerBase;
 class FNestedScopeIndexer;
 class FParametricTypeIndexer;
 struct FWriteIds;
-struct IBindIds;
+struct IDeclarations;
 
 class FWriter
 {
 public:
-	PLAINPROPS_API FWriter(const FIdIndexerBase& AllIds, const IBindIds& BindIds, const FBuiltSchemas& InSchemas, ESchemaFormat Format);
+	PLAINPROPS_API FWriter(const FIdIndexerBase& AllIds, const IDeclarations& Declarations, const FBuiltSchemas& InSchemas, ESchemaFormat Format);
 	PLAINPROPS_API ~FWriter();
 	
 	PLAINPROPS_API TConstArrayView<FNameId>			GetUsedNames() const;
diff --git a/Engine/Plugins/Experimental/PlainProps/Source/Tests/PlainPropsTests.cpp b/Engine/Plugins/Experimental/PlainProps/Source/Tests/PlainPropsTests.cpp
--- a/Engine/Plugins/Experimental/PlainProps/Source/Tests/PlainPropsTests.cpp
+++ b/Engine/Plugins/Experimental/PlainProps/Source/Tests/PlainPropsTests.cpp
@@ -2,26 +2,27 @@
 
 #if WITH_TESTS
 #include "Algo/Compare.h"
-#include "Logging/LogMacros.h"
-#include "Tests/TestHarnessAdapter.h"
 #include "Containers/AnsiString.h"
 #include "Containers/Map.h"
 #include "Containers/Set.h"
+#include "Containers/StaticArray.h"
 #include "Containers/StringConv.h"
 #include "Containers/StringView.h"
-#include "PlainPropsDiff.h"
-#include "PlainPropsRead.h"
-#include "PlainPropsInternalBuild.h"
-#include "PlainPropsWrite.h"
+#include "Logging/LogMacros.h"
 #include "PlainPropsBuildSchema.h"
+#include "PlainPropsDiff.h"
 #include "PlainPropsIndex.h"
-#include "PlainPropsVisualize.h"
+#include "PlainPropsInternalBuild.h"
 #include "PlainPropsInternalBuild.h"
 #include "PlainPropsInternalFormat.h"
 #include "PlainPropsInternalParse.h"
 #include "PlainPropsInternalPrint.h"
 #include "PlainPropsInternalRead.h"
+#include "PlainPropsRead.h"
+#include "PlainPropsVisualize.h"
+#include "PlainPropsWrite.h"
 #include "Templates/UnrealTemplate.h"
+#include "Tests/TestHarnessAdapter.h"
 
 DEFINE_LOG_CATEGORY_STATIC(LogPlainPropsTests, Log, All);
 
@@ -123,18 +124,18 @@ void AppendString(FUtf8StringBuilderBase& Out, const FAnsiString& Str)
 	Out.Append(Str);
 }
 
-class FTestBatchBuilder : public TIdIndexer<FAnsiString>, public IBindIds
+class FTestBatchBuilder final : public TIdIndexer<FAnsiString>, public IDeclarations
 {
 public:
-	FTestBatchBuilder(FScratchAllocator& InScratch) : Declarations(FDebugIds(*this)), Scratch(InScratch) {}
+	FTestBatchBuilder(FScratchAllocator& InScratch) : EnumDeclarations(FDebugIds(*this)), Scratch(InScratch) {}
 
 	FEnumId						DeclareEnum(FType Type, EEnumMode Mode, ELeafWidth Width, std::initializer_list<const char*> Names, std::initializer_list<uint64> Constants);
 	FEnumId						DeclareEnum(const char* Scope, const char* Name, EEnumMode Mode, ELeafWidth Width, std::initializer_list<const char*> Names, std::initializer_list<uint64> Constants);
-	FDeclId						DeclareStruct(FType Type, std::initializer_list<const char*> MemberOrder, EMemberPresence Occupancy, FOptionalDeclId Super = NoId);
-	FDeclId						DeclareStruct(const char* Scope, const char* Name, std::initializer_list<const char*> MemberOrder, EMemberPresence Occupancy, FOptionalDeclId Super = NoId);
+	FDeclId						DeclareStruct(FType Type, std::initializer_list<const char*> MemberOrder, TConstArrayView<FMemberSpec> MemberTypes, EMemberPresence Occupancy, FOptionalDeclId Super = NoId);
+	FDeclId						DeclareStruct(const char* Scope, const char* Name, std::initializer_list<const char*> MemberNames, TConstArrayView<FMemberSpec> MemberTypes, EMemberPresence Occupancy, FOptionalDeclId Super = NoId);
 	
-	const FEnumDeclaration&		Get(FEnumId Id) { return Declarations.Get(Id); }
-	const FStructDeclaration&	Get(FDeclId Id) { return Declarations.Get(Id); }
+	const FEnumDeclaration&		Get(FEnumId Id) const { return EnumDeclarations.Get(Id); }
+	const FStructDeclaration&	Get(FDeclId Id) const { return *StructDeclarations.FindChecked(Id); }
 
 	void						AddObject(FDeclId Id, FMemberBuilder&& Members);
 
@@ -144,7 +145,8 @@ public:
 
 private:
 	TArray<TPair<FDeclId, FBuiltStruct*>>		Objects;
-	FDeclarations								Declarations;
+	TMap<FDeclId, FStructDeclarationPtr>		StructDeclarations;
+	FEnumDeclarations							EnumDeclarations;
 	FScratchAllocator&							Scratch;
 
 	TArray<FMemberId> NameMembers(std::initializer_list<const char*> Members)
@@ -172,29 +174,31 @@ private:
 
 	TArray<char> GetNameData() const;
 
-	virtual FDeclId Lower(FBindId Id) const override final
+	virtual const FEnumDeclaration*		Find(FEnumId Id) const override { return &EnumDeclarations.Get(Id); }
+	virtual const FStructDeclaration*	Find(FStructId Id) const override { return StructDeclarations.FindChecked(FDeclId(Id)); }
+	virtual FDeclId						Lower(FBindId Id) const override
 	{
 		checkf(false, TEXT("All struct ids should be declared, nothing is bound with different names in this test suite"));
 		return LowerCast(Id);
 	}
 };
 
-FDeclId FTestBatchBuilder::DeclareStruct(FType Type, std::initializer_list<const char*> MemberOrder, EMemberPresence Occupancy, FOptionalDeclId Super)
+FDeclId FTestBatchBuilder::DeclareStruct(FType Type, std::initializer_list<const char*> MemberNames, TConstArrayView<FMemberSpec> MemberTypes,  EMemberPresence Occupancy, FOptionalDeclId Super)
 {
 	FDeclId Id = IndexDeclId(Type);
-	Declarations.DeclareStruct(Id, Type, 0, NameMembers(MemberOrder), Occupancy, Super);
+	StructDeclarations.Emplace(Id, Declare({Id, Super, /* v */ 0, Occupancy, NameMembers(MemberNames), MemberTypes}));
 	return Id;
 }
 
-FDeclId FTestBatchBuilder::DeclareStruct(const char* Scope, const char* Name, std::initializer_list<const char*> MemberOrder, EMemberPresence Occupancy, FOptionalDeclId Super)
+FDeclId FTestBatchBuilder::DeclareStruct(const char* Scope, const char* Name, std::initializer_list<const char*> MemberNames, TConstArrayView<FMemberSpec> MemberTypes, EMemberPresence Occupancy, FOptionalDeclId Super)
 {
-	return DeclareStruct(MakeType(Scope, Name), MemberOrder, Occupancy, Super);
+	return DeclareStruct(MakeType(Scope, Name), MemberNames, MemberTypes, Occupancy, Super);
 }
 
 FEnumId FTestBatchBuilder::DeclareEnum(FType Type, EEnumMode Mode, ELeafWidth, std::initializer_list<const char*> InNames, std::initializer_list<uint64> Constants)
 {
 	FEnumId Id = IndexEnum(Type);
-	Declarations.DeclareEnum(Id, Type, Mode, MakeEnumerators(InNames, Constants), EEnumAliases::Fail);
+	EnumDeclarations.Declare(Id, Type, Mode, MakeEnumerators(InNames, Constants), EEnumAliases::Fail);
 	return Id;
 }
 
@@ -205,13 +209,13 @@ FEnumId FTestBatchBuilder::DeclareEnum(const char* Scope, const char* Name, EEnu
 
 void FTestBatchBuilder::AddObject(FDeclId Id, FMemberBuilder&& Members)
 {
-	Objects.Emplace(Id, Members.BuildAndReset(Scratch, Declarations.Get(Id), GetDebug()));
+	Objects.Emplace(Id, Members.BuildAndReset(Scratch, Get(Id), GetDebug()));
 }
 
 TArray64<uint8> FTestBatchBuilder::Write()
 {
 	// Build partial schemas
-	FSchemasBuilder SchemaBuilders(Declarations, *this, *this, Scratch, ESchemaFormat::StableNames);
+	FSchemasBuilder SchemaBuilders(*this,  *this, Scratch, ESchemaFormat::StableNames);
 	for (const TPair<FDeclId, FBuiltStruct*>& Object : Objects)
 	{
 		SchemaBuilders.NoteStructAndMembers(Object.Key, *Object.Value);
@@ -481,15 +485,38 @@ TArray<OutType> MakeArray(const InType& Items)
 
 ////////////////////////////////////////////////////////////////////////////////////////////////
 
+template<uint32 N>
+TStaticArray<FMemberSpec, N> Respec(FMemberSpec Spec)
+{
+	return MakeUniformStaticArray<FMemberSpec, N>(Spec);
+}
+
+static FMemberSpec UniRange(FMemberSpec Item)	{ return FMemberSpec(ERangeSizeType::Uni, Item); }
+static FMemberSpec S8Range(FMemberSpec Item)	{ return FMemberSpec(ERangeSizeType::S8, Item); }
+static FMemberSpec U8Range(FMemberSpec Item)	{ return FMemberSpec(ERangeSizeType::U8, Item); }
+static FMemberSpec S16Range(FMemberSpec Item)	{ return FMemberSpec(ERangeSizeType::S16, Item); }
+static FMemberSpec U16Range(FMemberSpec Item)	{ return FMemberSpec(ERangeSizeType::U16, Item); }
+static FMemberSpec S32Range(FMemberSpec Item)	{ return FMemberSpec(ERangeSizeType::S32, Item); }
+static FMemberSpec U32Range(FMemberSpec Item)	{ return FMemberSpec(ERangeSizeType::U32, Item); }
+static FMemberSpec S64Range(FMemberSpec Item)	{ return FMemberSpec(ERangeSizeType::S64, Item); }
+static FMemberSpec U64Range(FMemberSpec Item)	{ return FMemberSpec(ERangeSizeType::U64, Item); }
+
+inline FMemberSpec Enum8(FEnumId Id)			{ return FMemberSpec({ELeafType::Enum, ELeafWidth::B8}, Id); }
+inline FMemberSpec Enum16(FEnumId Id)			{ return FMemberSpec({ELeafType::Enum, ELeafWidth::B16}, Id); }
+inline FMemberSpec Enum32(FEnumId Id)			{ return FMemberSpec({ELeafType::Enum, ELeafWidth::B32}, Id); }
+inline FMemberSpec Enum64(FEnumId Id)			{ return FMemberSpec({ELeafType::Enum, ELeafWidth::B64}, Id); }
+
+////////////////////////////////////////////////////////////////////////////////////////////////
+
 TEST_CASE_NAMED(FPlainPropsReadWriteTest, "System::Core::Serialization::PlainProps::ReadWrite", "[Core][PlainProps][SmokeFilter]")
 {
 	SECTION("Bool")
 	{
-		static constexpr std::initializer_list<const char*> MemberNames =  {"b0", "b1", "b2", "b3", "b4", "b5", "b6", "b7", "b8", "b9", "b10", "b11"};
+		static constexpr std::initializer_list<const char*> MemberNames =  {"b0", "b1", "b2", "b3", "b4", "b5", "b6", "b7", "b8", "b9", "b10", "b11"};	
 
 		TestSerialize([](FTestBatchBuilder& Batch, FScratchAllocator& Scratch)
 		{
-			FDeclId Id = Batch.DeclareStruct("Testing", "Bools", MemberNames, EMemberPresence::AllowSparse);
+			FDeclId Id = Batch.DeclareStruct("Testing", "Bools", MemberNames, Respec<12>(SpecBool), EMemberPresence::AllowSparse);
 
 			FMemberBuilder B1T;
 			B1T.Add(Batch.NameMember("b3"),	true);
@@ -582,7 +609,8 @@ TEST_CASE_NAMED(FPlainPropsReadWriteTest, "System::Core::Serialization::PlainPro
 
 		TestSerialize([](FTestBatchBuilder& Batch, FScratchAllocator& Scratch)
 		{
-			FDeclId Id = Batch.DeclareStruct("Test", "Numbers", MemberNames, EMemberPresence::AllowSparse);
+			FMemberSpec MemberTypes[] = {SpecF32, SpecF64, SpecS8, SpecU8, SpecS16, SpecU16, SpecS32, SpecU32, SpecS64, SpecU64};
+			FDeclId Id = Batch.DeclareStruct("Test", "Numbers", MemberNames, MemberTypes, EMemberPresence::AllowSparse);
 
 			FMemberBuilder Misc, Mins, Maxs, Some;
 
@@ -686,9 +714,9 @@ TEST_CASE_NAMED(FPlainPropsReadWriteTest, "System::Core::Serialization::PlainPro
 	{
 		TestSerialize([](FTestBatchBuilder& Batch, FScratchAllocator& Scratch)
 		{
-			FDeclId Char8Id  = Batch.DeclareStruct("Test", "Char8",  {"A", "B", "C", "D", "E", "F"}, EMemberPresence::AllowSparse);
-			FDeclId Char16Id = Batch.DeclareStruct("Test", "Char16", {"A", "B", "C", "D", "E", "F"}, EMemberPresence::AllowSparse);
-			FDeclId Char32Id = Batch.DeclareStruct("Test", "Char32", {"A", "B", "C", "D", "E", "F"}, EMemberPresence::AllowSparse);
+			FDeclId Char8Id  = Batch.DeclareStruct("Test", "Char8",  {"A", "B", "C", "D", "E", "F"}, Respec<6>(SpecUtf8),  EMemberPresence::AllowSparse);
+			FDeclId Char16Id = Batch.DeclareStruct("Test", "Char16", {"A", "B", "C", "D", "E", "F"}, Respec<6>(SpecUtf16), EMemberPresence::AllowSparse);
+			FDeclId Char32Id = Batch.DeclareStruct("Test", "Char32", {"A", "B", "C", "D", "E", "F"}, Respec<6>(SpecUtf32), EMemberPresence::AllowSparse);
 
 			FMemberBuilder Char8, Char16, Char32;
 
@@ -744,13 +772,16 @@ TEST_CASE_NAMED(FPlainPropsReadWriteTest, "System::Core::Serialization::PlainPro
 
 	SECTION("Dense")
 	{
-		static constexpr std::initializer_list<const char*> ExplicitMembers = {"A", "B", "C"};
-		static constexpr std::initializer_list<const char*> ImplicitMembers = {"0", "A", "1", "B", "2", "C", "3"};
+		static constexpr std::initializer_list<const char*> ExplicitNames = {"A", "B", "C"};
+		static constexpr std::initializer_list<const char*> ImplicitNames = {"0", "A", "1", "B", "2", "C", "3"};
 
 		TestSerialize([](FTestBatchBuilder& Batch, FScratchAllocator& Scratch)
 		{
-			FDeclId ExplicitId = Batch.DeclareStruct("Test", "ExplicitDense", ExplicitMembers, EMemberPresence::RequireAll);
-			FDeclId ImplicitId = Batch.DeclareStruct("Test", "ImplicitDense", ImplicitMembers, EMemberPresence::AllowSparse);
+			FMemberSpec ExplicitTypes[] = {SpecUtf8, SpecUtf16, SpecUtf32};
+			FMemberSpec ImplicitTypes[] = {SpecBool, SpecUtf8, SpecBool, SpecUtf16, SpecBool, SpecUtf32, SpecBool};
+
+			FDeclId ExplicitId = Batch.DeclareStruct("Test", "ExplicitDense", ExplicitNames, ExplicitTypes, EMemberPresence::RequireAll);
+			FDeclId ImplicitId = Batch.DeclareStruct("Test", "ImplicitDense", ImplicitNames, ImplicitTypes, EMemberPresence::AllowSparse);
 
 			FMemberBuilder X;
 			X.Add(Batch.NameMember("A"), char8_t('a'));
@@ -808,12 +839,12 @@ TEST_CASE_NAMED(FPlainPropsReadWriteTest, "System::Core::Serialization::PlainPro
 
 		TestSerialize([](FTestBatchBuilder& Batch, FScratchAllocator& Scratch)
 		{
-			FDeclId ObjectId = Batch.DeclareStruct("Test", "Object", ObjectMembers, EMemberPresence::AllowSparse);
-			FDeclId StructId = Batch.DeclareStruct("Test", "Struct", StructMembers, EMemberPresence::AllowSparse);
-			FDeclId NestedId = Batch.DeclareStruct("Test", "Nested", NestedMembers, EMemberPresence::AllowSparse);
-			FDeclId UnusedId = Batch.DeclareStruct("Test", "Unused", UnusedMembers, EMemberPresence::AllowSparse);
-			FDeclId EmptyId  = Batch.DeclareStruct("Test", "Empty",  EmptyMembers,  EMemberPresence::AllowSparse);
-
+			FDeclId NestedId = Batch.DeclareStruct("Test", "Nested", NestedMembers, {SpecS32, SpecS32}, EMemberPresence::AllowSparse);
+			FDeclId StructId = Batch.DeclareStruct("Test", "Struct", StructMembers, {NestedId, SpecBool}, EMemberPresence::AllowSparse);
+			FDeclId ObjectId = Batch.DeclareStruct("Test", "Object", ObjectMembers, {SpecF32, StructId, NestedId, SpecF32}, EMemberPresence::AllowSparse);
+			FDeclId UnusedId = Batch.DeclareStruct("Test", "Unused", UnusedMembers,	{SpecBool, SpecBool}, EMemberPresence::AllowSparse);
+			FDeclId EmptyId  = Batch.DeclareStruct("Test", "Empty",  EmptyMembers, {}, EMemberPresence::AllowSparse);
+		
 			FMemberBuilder Members;
 			Members.Add(Batch.NameMember("I1"), 100);
 			FBuiltStruct* NestedInStruct = Members.BuildAndReset(Scratch, Batch.Get(NestedId), Batch.GetDebug());
@@ -870,9 +901,8 @@ TEST_CASE_NAMED(FPlainPropsReadWriteTest, "System::Core::Serialization::PlainPro
 		TestSerialize([](FTestBatchBuilder& Batch, FScratchAllocator& Scratch)
 		{
 			// Test create holes in the original FNameId, FDeclId and FEnumId index range
-			FDeclId UnusedId = Batch.DeclareStruct("Test", "UnusedStruct", {"U1", "U2"}, EMemberPresence::AllowSparse);
-			FDeclId ObjectId = Batch.DeclareStruct("Test", "Enums", MemberNames, EMemberPresence::AllowSparse);
-
+			FDeclId UnusedId = Batch.DeclareStruct("Test", "UnusedStruct", {"U1", "U2"}, {SpecBool, SpecBool}, EMemberPresence::AllowSparse);
+			
 			FEnumId U	= Batch.DeclareEnum("Test", "UnusedEnum1",	EEnumMode::Flag, ELeafWidth::B8,	{"U3"}, {1}); // Hole
 			FEnumId A	= Batch.DeclareEnum("Test", "FlatDense8",	EEnumMode::Flat, ELeafWidth::B8,	{"A", "B", "C"}, {0, 1, 2});
 			FEnumId X	= Batch.DeclareEnum("Test", "UnusedEnum2",	EEnumMode::Flag, ELeafWidth::B8,	{"U4"}, {1}); // Hole
@@ -884,7 +914,12 @@ TEST_CASE_NAMED(FPlainPropsReadWriteTest, "System::Core::Serialization::PlainPro
 			FEnumId G	= Batch.DeclareEnum("Test", "FlatLimit32",	EEnumMode::Flat, ELeafWidth::B32,	{"Min", "Max"}, {0, 0xFFFFFFFF});
 			FEnumId H	= Batch.DeclareEnum("Test", "FlatLimit64",	EEnumMode::Flat, ELeafWidth::B64,	{"Min", "Max"}, {0, 0xFFFFFFFFFFFFFFFF});
 			FEnumId I	= Batch.DeclareEnum("Test", "FlagLimit64",	EEnumMode::Flag, ELeafWidth::B64,	{"One", "Max"}, {1, 0x8000000000000000});
-					
+			
+			FMemberSpec MemberTypes[] = {
+				Enum8(A), Enum8(A), Enum8(B), Enum8(B), Enum8(B), Enum8(B), 
+				Enum8(C), Enum8(D), Enum8(E), Enum16(F), Enum32(G), Enum64(H), Enum64(I) };
+			FDeclId ObjectId = Batch.DeclareStruct("Test", "Enums", MemberNames, MemberTypes, EMemberPresence::AllowSparse);
+
 			FMemberBuilder Members;
 			Members.AddEnum(Batch.NameMember("A2"),				A, (uint8)2);
 			Members.AddEnum(Batch.NameMember("A0"),				A, (uint8)0);
@@ -1000,12 +1035,15 @@ TEST_CASE_NAMED(FPlainPropsReadWriteTest, "System::Core::Serialization::PlainPro
 
 		TestSerialize([](FTestBatchBuilder& Batch, FScratchAllocator& Scratch)
 		{
+			FEnumId ABCD	= Batch.DeclareEnum("Test", "ABCD", EEnumMode::Flat, ELeafWidth::B16, {"A", "B", "C", "D"}, {0, 1, 2, 3});
+			FEnumId Unused1	= Batch.DeclareEnum("Test", "Unused1", EEnumMode::Flat, ELeafWidth::B8, {"X"}, {0});
+			FEnumId Unused2	= Batch.DeclareEnum("Test", "Unused2", EEnumMode::Flat, ELeafWidth::B8, {"Y"}, {0});
+
 			static constexpr std::initializer_list<const char*> MemberNames = { "B0", "B1", "B8", "B9", "D0", "D3", "Hi", "E3", "E0" };
-		
-			FDeclId ObjectId = Batch.DeclareStruct("Test", "Object", MemberNames, EMemberPresence::AllowSparse);
-			FEnumId Enum	= Batch.DeclareEnum("Test", "ABCD", EEnumMode::Flat, ELeafWidth::B16, {"A", "B", "C", "D"}, {0, 1, 2, 3});
-			FEnumId UnusedEnum1	= Batch.DeclareEnum("Test", "Unused1", EEnumMode::Flat, ELeafWidth::B8, {"X"}, {0});
-			FEnumId UnusedEnum2	= Batch.DeclareEnum("Test", "Unused2", EEnumMode::Flat, ELeafWidth::B8, {"Y"}, {0});
+			FMemberSpec MemberTypes[] = {	S32Range(SpecBool), S32Range(SpecBool), S64Range(SpecBool), S64Range(SpecBool), 
+											S32Range(SpecF64), S32Range(SpecF64), S32Range(SpecUtf8), 
+											S32Range(Enum16(ABCD)), S32Range(Enum8(Unused1)) };
+			FDeclId ObjectId = Batch.DeclareStruct("Test", "Object", MemberNames, MemberTypes, EMemberPresence::AllowSparse);
 		
 			FMemberBuilder Members;
 			Members.AddRange(Batch.NameMember("B0"), BuildLeafRange(Scratch, TConstArrayView<bool>()));
@@ -1015,8 +1053,8 @@ TEST_CASE_NAMED(FPlainPropsReadWriteTest, "System::Core::Serialization::PlainPro
 			Members.AddRange(Batch.NameMember("D0"), BuildLeafRange(Scratch, TConstArrayView<double>()));
 			Members.AddRange(Batch.NameMember("D3"), BuildLeafRange(Scratch, MakeArrayView({DBL_MIN, 0.0, DBL_MAX})));
 			Members.AddRange(Batch.NameMember("Hi"), BuildLeafRange(Scratch, MakeArrayView(u8"Hello!")));
-			Members.AddRange(Batch.NameMember("E3"), BuildEnumRange(Scratch, Enum, MakeArrayView({EABCD::B, EABCD::A, EABCD::D})));
-			Members.AddRange(Batch.NameMember("E0"), BuildEnumRange(Scratch, UnusedEnum1, TConstArrayView<EUnused1>()));
+			Members.AddRange(Batch.NameMember("E3"), BuildEnumRange(Scratch, ABCD, MakeArrayView({EABCD::B, EABCD::A, EABCD::D})));
+			Members.AddRange(Batch.NameMember("E0"), BuildEnumRange(Scratch, Unused1, TConstArrayView<EUnused1>()));
 
 			Batch.AddObject(ObjectId, MoveTemp(Members));
 		}, 
@@ -1051,8 +1089,8 @@ TEST_CASE_NAMED(FPlainPropsReadWriteTest, "System::Core::Serialization::PlainPro
 	{
 		TestSerialize([](FTestBatchBuilder& Batch, FScratchAllocator& Scratch)
 		{
-			FDeclId Utf8Id  = Batch.DeclareStruct("Test", "Utf8",  {"Null", "Empty", "Escape", "Latin1", "CJK", "Symbols"}, EMemberPresence::AllowSparse);
-			FDeclId Utf16Id = Batch.DeclareStruct("Test", "Utf16", {"Null", "Empty", "Escape", "Latin1", "CJK", "Symbols"}, EMemberPresence::AllowSparse);
+			FDeclId Utf8Id  = Batch.DeclareStruct("Test", "Utf8",  {"Null", "Empty", "Escape", "Latin1", "CJK", "Symbols"}, Respec<6>(S32Range(SpecUtf8)), EMemberPresence::AllowSparse);
+			FDeclId Utf16Id = Batch.DeclareStruct("Test", "Utf16", {"Null", "Empty", "Escape", "Latin1", "CJK", "Symbols"}, Respec<6>(S32Range(SpecUtf16)), EMemberPresence::AllowSparse);
 
 			FMemberBuilder Utf8, Utf16;
 
@@ -1099,8 +1137,8 @@ TEST_CASE_NAMED(FPlainPropsReadWriteTest, "System::Core::Serialization::PlainPro
 	{
 		TestSerialize([](FTestBatchBuilder& Batch, FScratchAllocator& Scratch)
 		{
-			FDeclId ObjectId = Batch.DeclareStruct("Test", "Object", {"Structs"}, EMemberPresence::AllowSparse);
-			FDeclId StructId = Batch.DeclareStruct("Test", "Struct", {"I", "F"}, EMemberPresence::AllowSparse);
+			FDeclId StructId = Batch.DeclareStruct("Test", "Struct", {"I", "F"}, {SpecS32, SpecF32}, EMemberPresence::AllowSparse);
+			FDeclId ObjectId = Batch.DeclareStruct("Test", "Object", {"Structs"}, {S32Range(StructId)}, EMemberPresence::AllowSparse);
 			 
 			FStructRangeBuilder Structs(4);
 			Structs[0].Add(Batch.NameMember("I"), 0);
@@ -1135,11 +1173,15 @@ TEST_CASE_NAMED(FPlainPropsReadWriteTest, "System::Core::Serialization::PlainPro
 
 		TestSerialize([](FTestBatchBuilder& Batch, FScratchAllocator& Scratch)
 		{
-			FDeclId Object = Batch.DeclareStruct("Test", "Object", {"IntRs", "EmptyRs", "FloatRs", "EnumRs", "UnusedEnumRs", "UnicodeRs", "StructRs", "StructRRs"}, EMemberPresence::AllowSparse);
-			FDeclId XY = Batch.DeclareStruct("Test", "XY", {"X", "Y"}, EMemberPresence::RequireAll);
-			FDeclId ZW = Batch.DeclareStruct("Test", "ZW", {"Z", "W"}, EMemberPresence::AllowSparse);
+			FDeclId XY = Batch.DeclareStruct("Test", "XY", {"X", "Y"}, {SpecF32, SpecF32}, EMemberPresence::RequireAll);
+			FDeclId ZW = Batch.DeclareStruct("Test", "ZW", {"Z", "W"}, {SpecF32, SpecF32}, EMemberPresence::AllowSparse);
 			FEnumId Enum = Batch.DeclareEnum("Test", "AB", EEnumMode::Flag, ELeafWidth::B8, {"A", "B"}, {1, 4});
 			FEnumId UnusedEnum = Batch.DeclareEnum("Test", "Unused", EEnumMode::Flat, ELeafWidth::B8, {"X"}, {0});
+			static constexpr std::initializer_list<const char*> MemberNames = { "IntRs", "EmptyRs", "FloatRs", "EnumRs", "UnusedEnumRs", "UnicodeRs", "StructRs", "StructRRs" };
+			FMemberSpec MemberTypes[] = { S32Range(S32Range(SpecS32)), S32Range(S32Range(SpecBool)), S64Range(S32Range(SpecF32)), 
+				U8Range(S32Range(Enum8(Enum))), U8Range(S32Range(Enum8(UnusedEnum))), S32Range(S32Range(SpecUtf8)),
+				U64Range(U64Range(XY)), U32Range(S16Range(S16Range(ZW)))};
+			FDeclId Object = Batch.DeclareStruct("Test", "Object", MemberNames, MemberTypes, EMemberPresence::AllowSparse);
 	
 			FNestedRangeBuilder IntRs(MakeLeafRangeSchema<int32, int32>(), 3);
 			IntRs.Add(BuildLeafRange(Scratch, MakeArrayView({1})));
@@ -1258,8 +1300,10 @@ TEST_CASE_NAMED(FPlainPropsReadWriteTest, "System::Core::Serialization::PlainPro
 	{
 		TestSerialize([](FTestBatchBuilder& Batch, FScratchAllocator& Scratch)
 		{
-			FDeclId Object = Batch.DeclareStruct("Test", "Object", {"Bools", "Structs", "BF", "BT" }, EMemberPresence::AllowSparse);
-			FDeclId Struct = Batch.DeclareStruct("Test", "Struct", {"MaybeB", "Bs", "MaybeBs", "B"}, EMemberPresence::AllowSparse);
+			FMemberSpec StructMemberTypes[] = { UniRange(SpecBool), S32Range(SpecBool), UniRange(SpecBool), SpecBool };
+			FDeclId Struct = Batch.DeclareStruct("Test", "Struct", {"MaybeB", "Bs", "MaybeBs", "B"}, StructMemberTypes, EMemberPresence::AllowSparse);
+			FMemberSpec ObjectMemberTypes[] = { UniRange(SpecBool), S32Range(Struct), SpecBool, SpecBool };
+			FDeclId Object = Batch.DeclareStruct("Test", "Object", {"Bools", "Structs", "BF", "BT" }, ObjectMemberTypes, EMemberPresence::AllowSparse);
 						
 			const bool True = true;
 			const bool False = false;
@@ -1316,12 +1360,12 @@ TEST_CASE_NAMED(FPlainPropsReadWriteTest, "System::Core::Serialization::PlainPro
 	{
 		TestSerialize([](FTestBatchBuilder& Batch, FScratchAllocator& Scratch)
 		{
-			FDeclId Unused1 = Batch.DeclareStruct("Test", "Unused1", {"X"}, EMemberPresence::AllowSparse);
-			FDeclId SA = Batch.DeclareStruct("Test", "SA", {"X"}, EMemberPresence::AllowSparse);
-			FDeclId Unused2 = Batch.DeclareStruct("Test", "Unused2", {"X"}, EMemberPresence::AllowSparse);
-			FDeclId SB = Batch.DeclareStruct("Test", "SB", {"X"}, EMemberPresence::AllowSparse);
-			FDeclId Object = Batch.DeclareStruct("Test", "Object", {"Same", "Some", "None", "Diff"}, EMemberPresence::AllowSparse);
-			FDeclId Unused3 = Batch.DeclareStruct("Test", "Unused3", {"X"}, EMemberPresence::AllowSparse);
+			FDeclId Unused1 = Batch.DeclareStruct("Test", "Unused1", {"X"}, {SpecS32}, EMemberPresence::AllowSparse);
+			FDeclId SA = Batch.DeclareStruct("Test", "SA", {"X"}, {SpecS32}, EMemberPresence::AllowSparse);
+			FDeclId Unused2 = Batch.DeclareStruct("Test", "Unused2", {"X"}, {SpecF32}, EMemberPresence::AllowSparse);
+			FDeclId SB = Batch.DeclareStruct("Test", "SB", {"X"}, {SpecF32}, EMemberPresence::AllowSparse);
+			FDeclId Object = Batch.DeclareStruct("Test", "Object", {"Same", "Some", "None", "Diff"}, {SA, SA, SpecBool, SpecDynamicStruct}, EMemberPresence::AllowSparse);
+			FDeclId Unused3 = Batch.DeclareStruct("Test", "Unused3", {"X"}, {SpecBool}, EMemberPresence::AllowSparse);
 		
 			auto BuildStruct = [&](FDeclId Struct, auto X)
 				{
@@ -1365,10 +1409,11 @@ TEST_CASE_NAMED(FPlainPropsReadWriteTest, "System::Core::Serialization::PlainPro
 	{
 		TestSerialize([](FTestBatchBuilder& Batch, FScratchAllocator& Scratch)
 		{
-			FDeclId SA = Batch.DeclareStruct("Test", "SA", {"X"}, EMemberPresence::AllowSparse);
-			FDeclId Unused = Batch.DeclareStruct("Test", "Unused2", {"X"}, EMemberPresence::AllowSparse);
-			FDeclId SB = Batch.DeclareStruct("Test", "SB", {"X"}, EMemberPresence::AllowSparse);
-			FDeclId Object = Batch.DeclareStruct("Test", "Object", {"Same", "Some", "None", "Diff", "SameEmpty", "DiffEmpty", "DiffNested"}, EMemberPresence::AllowSparse);
+			FDeclId SA = Batch.DeclareStruct("Test", "SA", {"X"}, {S32Range(SpecS32)}, EMemberPresence::AllowSparse);
+			FDeclId Unused = Batch.DeclareStruct("Test", "Unused2", {"X"}, {S32Range(SpecS32)}, EMemberPresence::AllowSparse);
+			FDeclId SB = Batch.DeclareStruct("Test", "SB", {"X"}, {S32Range(SpecF32)}, EMemberPresence::AllowSparse);
+			FMemberSpec MemberTypes[] = { S32Range(SA), S32Range(SA), SpecBool, S32Range(SpecDynamicStruct), S32Range(SA), S32Range(SpecDynamicStruct), S32Range(S32Range(SpecDynamicStruct)) };
+			FDeclId Object = Batch.DeclareStruct("Test", "Object", {"Same", "Some", "None", "Diff", "SameEmpty", "DiffEmpty", "DiffNested"}, MemberTypes, EMemberPresence::AllowSparse);
 		
 			auto BuildStructRange = [&](FDeclId Struct, auto X)
 				{
@@ -1428,10 +1473,10 @@ TEST_CASE_NAMED(FPlainPropsReadWriteTest, "System::Core::Serialization::PlainPro
 	{
 		TestSerialize([](FTestBatchBuilder& Batch, FScratchAllocator& Scratch)
 		{
-			FDeclId Unused = Batch.DeclareStruct("Test", "X", {"X"}, EMemberPresence::AllowSparse);
-			FDeclId Low = Batch.DeclareStruct("Test", "Low", {"LInt"}, EMemberPresence::AllowSparse);
-			FDeclId Mid = Batch.DeclareStruct("Test", "Mid", {"MInt", "MLow"}, EMemberPresence::AllowSparse, ToOptional(Low));
-			FDeclId Top = Batch.DeclareStruct("Test", "Top", {"TInt", "TLow", "TMids"}, EMemberPresence::AllowSparse, ToOptional(Mid));
+			FDeclId Unused = Batch.DeclareStruct("Test", "X", {"X"}, {SpecBool}, EMemberPresence::AllowSparse);
+			FDeclId Low = Batch.DeclareStruct("Test", "Low", {"LInt"}, {SpecS32}, EMemberPresence::AllowSparse);
+			FDeclId Mid = Batch.DeclareStruct("Test", "Mid", {"MInt", "MLow"}, {SpecS32, Low}, EMemberPresence::AllowSparse, ToOptional(Low));
+			FDeclId Top = Batch.DeclareStruct("Test", "Top", {"TInt", "TLow", "TMids"}, {SpecS32, Low, S32Range(Mid)}, EMemberPresence::AllowSparse, ToOptional(Mid));
 			
 			FMemberBuilder Members;
 			Members.Add(Batch.NameMember("LInt"), 123);
@@ -1508,15 +1553,15 @@ TEST_CASE_NAMED(FPlainPropsReadWriteTest, "System::Core::Serialization::PlainPro
 	SECTION("SparseInheritance")
 	{
 		TestSerialize([](FTestBatchBuilder& Batch, FScratchAllocator& Scratch)
-		{																										// Usage by A B C 
-			FDeclId B0 = Batch.DeclareStruct("Test", "B0", {"0"}, EMemberPresence::AllowSparse);				 // - - -
-			FDeclId B1 = Batch.DeclareStruct("Test", "B1", {"1"}, EMemberPresence::AllowSparse, ToOptional(B0)); // 1 1 1
-			FDeclId B2 = Batch.DeclareStruct("Test", "B2", {"2"}, EMemberPresence::AllowSparse, ToOptional(B1)); // - 1 1
-			FDeclId B3 = Batch.DeclareStruct("Test", "B3", {"3"}, EMemberPresence::AllowSparse, ToOptional(B2)); // - - 0
-			FDeclId B4 = Batch.DeclareStruct("Test", "B4", {"4"}, EMemberPresence::AllowSparse, ToOptional(B3)); // 1 1 1
-			FDeclId B5 = Batch.DeclareStruct("Test", "B5", {"5"}, EMemberPresence::AllowSparse, ToOptional(B4)); // 1 1 0
-			FDeclId B6 = Batch.DeclareStruct("Test", "B6", {   }, EMemberPresence::AllowSparse, ToOptional(B5)); // 0 - 0
-			FDeclId C5 = Batch.DeclareStruct("Test", "C5", {"5"}, EMemberPresence::AllowSparse, ToOptional(B4)); // - - -
+		{																													// A B C usage
+			FDeclId B0 = Batch.DeclareStruct("Test", "B0", {"0"}, {SpecS32}, EMemberPresence::AllowSparse);					// - - -
+			FDeclId B1 = Batch.DeclareStruct("Test", "B1", {"1"}, {SpecS32}, EMemberPresence::AllowSparse, ToOptional(B0)); // 1 1 1
+			FDeclId B2 = Batch.DeclareStruct("Test", "B2", {"2"}, {SpecS32}, EMemberPresence::AllowSparse, ToOptional(B1)); // - 1 1
+			FDeclId B3 = Batch.DeclareStruct("Test", "B3", {"3"}, {SpecS32}, EMemberPresence::AllowSparse, ToOptional(B2)); // - - 0
+			FDeclId B4 = Batch.DeclareStruct("Test", "B4", {"4"}, {SpecS32}, EMemberPresence::AllowSparse, ToOptional(B3)); // 1 1 1
+			FDeclId B5 = Batch.DeclareStruct("Test", "B5", {"5"}, {SpecS32}, EMemberPresence::AllowSparse, ToOptional(B4)); // 1 1 0
+			FDeclId B6 = Batch.DeclareStruct("Test", "B6", {   }, {       }, EMemberPresence::AllowSparse, ToOptional(B5)); // 0 - 0
+			FDeclId C5 = Batch.DeclareStruct("Test", "C5", {"5"}, {SpecS32}, EMemberPresence::AllowSparse, ToOptional(B4)); // - - -
 			
 			FMemberBuilder A;
 			A.Add(Batch.NameMember("1"), 1);
@@ -1623,14 +1668,14 @@ TEST_CASE_NAMED(FPlainPropsReadWriteTest, "System::Core::Serialization::PlainPro
 			FType S1T{NestedUnused1,	Batch.MakeTypename("S1")};
 			FType S2T{NestedUsed,		Batch.MakeTypename("S2")};
 			FType S3T = Batch.MakeParametricType({NestedUnused2,	Batch.MakeTypename("S3")}, {S1T});
-			FType S4T = Batch.MakeParametricType({DoubleNested,	Batch.MakeTypename("S4")}, {S2T, E2T});
+			FType S4T = Batch.MakeParametricType({DoubleNested,		Batch.MakeTypename("S4")}, {S2T, E2T});
 			FType S5T = Batch.MakeParametricType({NestedUnused3,	Batch.MakeTypename("S5")}, {E3T, E1T, S2T});
 			
-			FDeclId S1D = Batch.DeclareStruct(S1T, {"M1"}, EMemberPresence::AllowSparse);
-			FDeclId S2D = Batch.DeclareStruct(S2T, {"M2"}, EMemberPresence::AllowSparse);
-			FDeclId S3D = Batch.DeclareStruct(S3T, {"M3"}, EMemberPresence::AllowSparse);
-			FDeclId S4D = Batch.DeclareStruct(S4T, {"M4"}, EMemberPresence::AllowSparse);
-			FDeclId S5D = Batch.DeclareStruct(S5T, {"M5"}, EMemberPresence::AllowSparse);
+			FDeclId S1D = Batch.DeclareStruct(S1T, {"M1"}, {SpecS32}, EMemberPresence::AllowSparse);
+			FDeclId S2D = Batch.DeclareStruct(S2T, {"M2"}, {SpecS32}, EMemberPresence::AllowSparse);
+			FDeclId S3D = Batch.DeclareStruct(S3T, {"M3"}, {SpecS32}, EMemberPresence::AllowSparse);
+			FDeclId S4D = Batch.DeclareStruct(S4T, {"M4"}, {SpecS32}, EMemberPresence::AllowSparse);
+			FDeclId S5D = Batch.DeclareStruct(S5T, {"M5"}, {SpecS32}, EMemberPresence::AllowSparse);
 
 			FMemberBuilder Members;
 			Members.Add(Batch.NameMember("M4"), 1);
diff --git a/Engine/Plugins/Experimental/PlainProps/Source/Tests/PlainPropsUeCoreTests.cpp b/Engine/Plugins/Experimental/PlainProps/Source/Tests/PlainPropsUeCoreTests.cpp
--- a/Engine/Plugins/Experimental/PlainProps/Source/Tests/PlainPropsUeCoreTests.cpp
+++ b/Engine/Plugins/Experimental/PlainProps/Source/Tests/PlainPropsUeCoreTests.cpp
@@ -35,7 +35,7 @@ namespace PlainProps::UE::Test
 
 static TIdIndexer<FSensitiveName>	GNames;
 static FDebugIds					GDebug(GNames);
-static FDeclarations				GTypes(GDebug);
+static FEnumDeclarations			GEnums(GDebug);
 static FSchemaBindings				GSchemas(GDebug);
 static FCustomBindingsBottom		GCustoms(GDebug);
 static FCustomBindingsOverlay		GDeltaCustoms(GCustoms);
@@ -60,7 +60,7 @@ struct FDefaultRuntime
 	using Ids = FRuntimeIds;
 	template<class T> using CustomBindings = TCustomBind<T>;
 
-	static FDeclarations&			GetTypes()			{ return GTypes; }
+	static FEnumDeclarations&		GetEnums()			{ return GEnums; }
 	static FSchemaBindings&			GetSchemas()		{ return GSchemas; }
 	static FCustomBindings&			GetCustoms()		{ return GCustoms; }
 };
@@ -83,8 +83,8 @@ struct TScopedEnumDeclaration
 	using Ctti = CttiOf<Enum>;
 
 	FEnumId Id;
-	TScopedEnumDeclaration() : Id(DeclareNativeEnum<Ctti, Ids>(Runtime::GetTypes(), Mode)) {}
-	~TScopedEnumDeclaration() { Runtime::GetTypes().DropEnum(Id); }
+	TScopedEnumDeclaration() : Id(DeclareNativeEnum<Ctti, Ids>(Runtime::GetEnums(), Mode)) {}
+	~TScopedEnumDeclaration() { Runtime::GetEnums().Drop(Id); }
 };
 
 template<class T>
@@ -92,31 +92,27 @@ using TScopedDefaultStructBinding = TScopedStructBinding<T, FDefaultRuntime>;
 
 ////////////////////////////////////////////////////////////////////////////////////////////////
 
-// Maybe replace with macro, e.g. PP_DECLARE_CUSTOM_DENSE_STRUCT(FRuntimeIds, FName, void, Idx)
-struct FNameDeclaration 
+
+struct FNameBinding final : ICustomBinding
 {
-	FDeclId				Id;
-	FMemberId			Idx;
+	const FDualStructId	Id;
+	const FMemberId		IdxName;
+	TSet<FName>			Names;
 
-	FNameDeclaration(FType Type = IndexStructName<FRuntimeIds, ETypename::Decl, TTypename<FName>>())
-	: Id(FRuntimeIds::IndexStruct(Type))
-	, Idx(FRuntimeIds::IndexMember("Idx"))
-	{
-		GTypes.DeclareStruct(Id, Type, 123, MakeArrayView(&Idx, 1), EMemberPresence::RequireAll);
-	}
+	FNameBinding()
+	: Id(IndexStructDualId<FRuntimeIds, TTypename<FName>>())
+	, IdxName(FRuntimeIds::IndexMember("Idx"))
+	{}
 
-	~FNameDeclaration()
+	FStructDeclarationPtr Declare() const
 	{
-		GTypes.DropStructRef(Id);
+		return PlainProps::Declare({Id, NoId, 0, EMemberPresence::RequireAll, {IdxName}, {Specify<int32>()}});
 	}
-};
 
-struct FNameBinding final : ICustomBinding
-{
 	virtual void SaveCustom(FMemberBuilder& Dst, const void* Src, const void*, const FSaveContext& Ctx) override
 	{
 		FSetElementId Idx = Names.Add(*static_cast<const FName*>(Src));
-		Dst.Add(Declaration.Idx, Idx.AsInteger());
+		Dst.Add(IdxName, Idx.AsInteger());
 	}
 
 	virtual void LoadCustom(void* Dst, FStructLoadView Src, ECustomLoadMethod) const override
@@ -131,9 +127,6 @@ struct FNameBinding final : ICustomBinding
 		FName B = *static_cast<const FName*>(StructB);
 		return A.IsEqual(B, ENameCase::CaseSensitive);
 	}
-
-	FNameDeclaration	Declaration;
-	TSet<FName>			Names;
 };
 
 ////////////////////////////////////////////////////////////////////////////////////////////////
@@ -142,6 +135,7 @@ class FBatchSaver
 {
 public:
 	explicit FBatchSaver(const FCustomBindings& CustomBase);
+	virtual ~FBatchSaver() {}
 
 	template<class T>
 	void						Save(T&& Object);
@@ -161,21 +155,21 @@ private:
 FBatchSaver::FBatchSaver(const FCustomBindings& Underlay)
 : Customs(Underlay)
 {
-	Customs.BindStruct(UpCast(SavedNames.Declaration.Id), SavedNames, GTypes.Get(SavedNames.Declaration.Id), {});
+	Customs.BindStruct(UpCast(SavedNames.Id), SavedNames, SavedNames.Declare(), {});
 }
 
 template<class T>
 void FBatchSaver::Save(T&& Object) 
 {
 	FBothStructId Id = IndexStructBothId<FRuntimeIds, TTypename<std::remove_reference_t<T>>>();
-	SavedObjects.Emplace(Id.DeclId, SaveStruct(&Object, Id.BindId, {{GTypes, GSchemas, Customs}, Scratch}));
+	SavedObjects.Emplace(Id.DeclId, SaveStruct(&Object, Id.BindId, {{GSchemas, Customs}, Scratch}));
 }
 
 template<class T>
 bool FBatchSaver::SaveDelta(const T& Object, const T& Default) 
 {
 	FBothStructId Id = IndexStructBothId<FRuntimeIds, TTypename<std::remove_reference_t<T>>>();
-	if (const FBuiltStruct* Delta = SaveStructDeltaIfDiff(&Object, &Default, Id.BindId, {{GTypes, GSchemas, Customs}, Scratch}))
+	if (const FBuiltStruct* Delta = SaveStructDeltaIfDiff(&Object, &Default, Id.BindId, {{GSchemas, Customs}, Scratch}))
 	{
 		SavedObjects.Emplace(Id.DeclId, MoveTemp(Delta));
 		return true;
@@ -204,8 +198,8 @@ TArray64<uint8> FBatchSaver::Write(TArray<FStructId>* OutMemoryIds) const
 	ESchemaFormat Format = OutMemoryIds ? ESchemaFormat::InMemoryNames : ESchemaFormat::StableNames;
 
 	// Build partial schemas
-	const FStructBindIds BindIds(Customs, GSchemas);
-	FSchemasBuilder SchemaBuilders(GTypes, GNames, BindIds, Scratch, Format);
+	const FBindDeclarations Types(GEnums, Customs, GSchemas);
+	FSchemasBuilder SchemaBuilders(GNames, Types,  Scratch, Format);
 	for (const IdBuiltStructPair& Object : SavedObjects)
 	{
 		SchemaBuilders.NoteStructAndMembers(Object.Key, *Object.Value);
@@ -217,7 +211,7 @@ TArray64<uint8> FBatchSaver::Write(TArray<FStructId>* OutMemoryIds) const
 	}
 
 	// Filter out declared but unused names and ids
-	FWriter Writer(GNames, BindIds, Schemas, Format);
+	FWriter Writer(GNames, Types, Schemas, Format);
 	TArray<FSensitiveName> UsedNames;
 	UsedNames.Reserve(Writer.GetUsedNames().Num());
 	for (FNameId Name : Writer.GetUsedNames())
@@ -406,12 +400,12 @@ public:
 
 		// Read names and bind custom loader
 		NameBinding.Names.Append(GrabNumAndArray<FName>(It));
-		Customs.BindStruct(UpCast(NameBinding.Declaration.Id), NameBinding, GTypes.Get(NameBinding.Declaration.Id), {});
+		Customs.BindStruct(UpCast(NameBinding.Id), NameBinding, NameBinding.Declare(), {});
 		CHECK(It.Grab<uint32>() == Magics[5]);
 
 		// Finally create load plans
 		TConstArrayView<FStructId> LoadStructIds = RuntimeIds ? RuntimeIds->Translation.GetStructIds(SavedSchemas->NumStructSchemas) : InRuntimeIds;
-		Plans = CreateLoadPlans(Batch, GTypes, Customs, GSchemas, LoadStructIds, Format);
+		Plans = CreateLoadPlans(Batch, Customs, GSchemas, LoadStructIds, Format);
 
 		// Create BatchIds for the loading phase
 		if (RuntimeIds)
@@ -652,9 +646,11 @@ struct FNDCBinding : ICustomBinding
 	const FMemberId MemberIds[1];
 
 	template<class Ids>
-	FNDCBinding(TCustomInit<Ids>)
+	FNDCBinding(TCustomSpecifier<Ids, 1>& Spec)
 	: MemberIds{Ids::IndexMember("X")}
-	{}
+	{
+		Spec.Members[0] = Specify<int32>();
+	}
 
 	void Save(FMemberBuilder& Dst, const FNDC& Src, const FNDC*, const FSaveContext&) const
 	{
diff --git a/Engine/Plugins/Experimental/PlainPropsUObject/Source/Private/PlainPropsUObjectBindings.cpp b/Engine/Plugins/Experimental/PlainPropsUObject/Source/Private/PlainPropsUObjectBindings.cpp
--- a/Engine/Plugins/Experimental/PlainPropsUObject/Source/Private/PlainPropsUObjectBindings.cpp
+++ b/Engine/Plugins/Experimental/PlainPropsUObject/Source/Private/PlainPropsUObjectBindings.cpp
@@ -89,6 +89,8 @@ namespace PlainProps::UE
 
 static constexpr ERangeSizeType DefaultRangeMax = RangeSizeOf(FDefaultAllocator::SizeType{});
 
+inline FMemberSpec DefaultRangeOf(FMemberSpec Item) { return FMemberSpec(DefaultRangeMax, Item); }
+
 ////////////////////////////////////////////////////////////////////////////////////////////////
 
 struct FDefaultStruct
@@ -317,7 +319,7 @@ FCommonMemberIds::FCommonMemberIds(TIdIndexer<FSensitiveName>& Names)
 ////////////////////////////////////////////////////////////////////////////////////////////////
 
 FGlobals::FGlobals() 
-: Types(FDebugIds(Names))
+: Enums(FDebugIds(Names))
 , Schemas(FDebugIds(Names))
 , Customs(FDebugIds(Names))
 , Scopes(Names)
@@ -374,56 +376,6 @@ static bool HasAny(uint64 Flags)
 
 ////////////////////////////////////////////////////////////////////////////////////////////////
 
-static bool ShouldBind(const UStruct* Struct);
-
-inline bool ShouldBind(const FProperty* Property)
-{
-	if (HasAny<CASTCLASS_FStructProperty>(Property->GetCastFlags()))
-	{
-		return ShouldBind(static_cast<const FStructProperty*>(Property)->Struct);
-	}
-
-	return true;
-}
-
-static bool ShouldBind(const UStruct* Struct)
-{
-	for (FProperty* Property = Struct->PropertyLink; Property; Property = Property->PropertyLinkNext)
-	{
-		if (ShouldBind(Property))
-		{
-			return true;
-		}
-	}
-	return false;
-}
-
-// Serialize property-less UObject and UScriptStructs as their super class.
-//
-// E.g. FVector_NetQuantize10 is a pure runtime abstraction serialized as FVector.
-//		FAttenuationSubmixSendSettings is just a FSoundSubmixSendInfoBase but has a different 
-//		default constructor that matters in sparse delta serialization.
-//		UObjects are never instantiated during serialization so can be safely simplified.
-//
-// These heuristics might need more tuning
-static const UStruct* SkipEmptyBases(const UStruct* In)
-{
-	const UStruct* FirstOwner = In->PropertyLink ? In->PropertyLink->GetOwnerChecked<UStruct>() : In;
-	if (In != FirstOwner)
-	{
-		if (const UScriptStruct* Struct = Cast<const UScriptStruct>(In))
-		{
-			if (!(Struct->StructFlags & STRUCT_ZeroConstructor))
-			{
-				return In;
-			}
-		}
-
-		return FirstOwner;
-	}
-	return In;
-}
-
 static FType IndexType(const UField* Field)
 {
 	check(Field);
@@ -439,20 +391,6 @@ static FType IndexType(const UField* Field)
 	return { GUE.Names.NestReversedScopes(ReversedOuters), Name };
 }
 
-static FOptionalDeclId IndexSuper(const UStruct* Struct)
-{
-	if (const UStruct* Super = Struct->GetInheritanceSuper())
-	{
-		if (ShouldBind(Super))
-		{
-			const UStruct* NonEmptySuper = SkipEmptyBases(Super);
-			return GUE.Names.IndexDeclId(IndexType(NonEmptySuper));
-		}
-	}
-	
-	return NoId;
-}
-
 static EMemberPresence GetOccupancy(const UStruct* Struct)
 {
 	if (Struct->HasAnyCastFlags(CASTCLASS_UScriptStruct))
@@ -463,57 +401,6 @@ static EMemberPresence GetOccupancy(const UStruct* Struct)
 	return EMemberPresence::AllowSparse;
 }
 
-using FMemberArray = TArray<FMemberId, TInlineAllocator<64>>;
-
-inline void DeclareMembers(FMemberArray& Out, const UStruct* Struct)
-{
-	for (FProperty* It = Struct->PropertyLink; It && It->GetOwner<UStruct>() == Struct; It = It->PropertyLinkNext)
-	{
-		if (ShouldBind(It))
-		{
-			Out.Add(GUE.Names.NameMember(It->GetFName()));
-		}
-	}
-}
-
-// Must match BindSuperMembers
-static void DeclareSuperMembers(FMemberArray& Out, const UStruct* Struct)
-{
-	if (const UStruct* Super = Struct->GetInheritanceSuper())
-	{
-		DeclareSuperMembers(Out, Super);
-		if (ShouldBind(Super))
-		{
-			DeclareMembers(Out, Super);
-		}
-	}
-}
-
-static void DeclareStruct(const UStruct* Struct, FType Type, FDeclId Id)
-{
-	FOptionalDeclId Super = IndexSuper(Struct);
-	EMemberPresence Occupancy = GetOccupancy(Struct);
-	
-	FMemberArray Members;
-	if (Super && Occupancy == EMemberPresence::RequireAll)
-	{
-		// Flatten inheritance chain for dense structs
-		Super = NoId;
-		DeclareSuperMembers(/* out */ Members, Struct);
-	}
-	DeclareMembers(/* out */ Members, Struct);
-
-	GUE.Types.DeclareStruct(Id, Type, 0, Members, Occupancy, Super);
-}
-
-static FDeclId DeclareStruct(const UStruct* Struct)
-{
-	FType Type = IndexType(Struct);
-	FDeclId Id = GUE.Names.IndexDeclId(Type);
-	DeclareStruct(Struct, Type, Id);
-	return Id;
-}
-
 ////////////////////////////////////////////////////////////////////////////////////////////////
 
 static FTypedRange SaveNames(TConstArrayView<FName> Names, const FSaveContext& Ctx)
@@ -537,6 +424,13 @@ static void LoadNames(TArray<FName>& Dst, FStructRangeLoadView Src)
 	}
 }
 
+FFieldPathBinding::FFieldPathBinding(TPropertySpecifier<2>& Spec)
+: MemberIds{GUE.Members.Path, GUE.Members.Owner}
+{
+	Spec.Members[0] = DefaultRangeOf(FDeclId(GUE.Structs.Name));
+	Spec.Members[1] = FDeclId(GUE.Structs.WeakObjectPtr);
+}
+
 void FFieldPathBinding::Save(FMemberBuilder& Dst, const FFieldPath& Src, const FFieldPath*, const FSaveContext& Ctx) const
 {
 	Dst.AddRange(GUE.Members.Path, SaveNames(Src.*GFieldPathPath, Ctx));
@@ -558,6 +452,12 @@ bool FFieldPathBinding::Diff(const FFieldPath& A, const FFieldPath& B, const FBi
 
 ////////////////////////////////////////////////////////////////////////////////////////////////
 
+FDelegateBinding::FDelegateBinding(TPropertySpecifier<2>& Out)
+: MemberIds{GUE.Members.Object, GUE.Members.Function}
+{
+	Out.SetMembers(GUE.Structs.WeakObjectPtr, GUE.Structs.Name);
+}
+
 void FDelegateBinding::Save(FMemberBuilder& Dst, const FScriptDelegate& Src, const FScriptDelegate* Default, const FSaveContext& Ctx) const
 {
 	FDelegateAccess::FReadScope Scope(Src);
@@ -618,7 +518,7 @@ static FTypedRange SaveInvocations(const FMulticastInvocationList& In, const FSa
 	
 		if (int32 NumKept = Keep.CountSetBits())
 		{
-			const FStructDeclaration& ItemDecl = GUE.Types.Get(LowerCast(ItemId));
+			const FStructDeclaration& ItemDecl = Ctx.GetDeclaration(ItemId);
 			const FUnicastScriptDelegate* Src = In.GetData();
 			FStructRangeSaver Dst(Ctx.Scratch, static_cast<uint64>(NumKept));
 			FMemberBuilder Tmp;
@@ -637,6 +537,11 @@ static FTypedRange SaveInvocations(const FMulticastInvocationList& In, const FSa
 	return {Schema, nullptr};
 }
 
+static FStructDeclarationPtr DeclareMulticastDelegate()
+{
+	return Declare({GUE.Structs.MulticastDelegate, NoId, 0, EMemberPresence::RequireAll, {GUE.Members.Invocations}, {DefaultRangeOf(FDeclId(GUE.Structs.Delegate))}});
+}
+
 static void SaveMulticastDelegate(FMemberBuilder& Dst, const FMulticastScriptDelegate& Src, const FSaveContext& Ctx)
 {
 	FDelegateAccess::FReadScope Scope(Src);
@@ -693,6 +598,12 @@ static bool DiffMulticastDelegate(const FMulticastScriptDelegate& A, const FMult
 	return DiffInvocations(A.*GMulticastDelegateInvocationList, B.*GMulticastDelegateInvocationList);
 }
 
+FMulticastInlineDelegateBinding::FMulticastInlineDelegateBinding(TPropertySpecifier<1>& Spec)
+: MemberIds{GUE.Members.Invocations}
+{
+	Spec.Members[0] = DefaultRangeOf(FDeclId(GUE.Structs.Delegate));
+}
+
 void FMulticastInlineDelegateBinding::Save(FMemberBuilder& Dst, const FMulticastScriptDelegate& Src, const FMulticastScriptDelegate* Default, const FSaveContext& Ctx) const
 {
 	SaveMulticastDelegate(Dst, Src, Ctx);
@@ -805,14 +716,21 @@ static FBindId BindSparseDelegate(FBindId Owner, FMulticastSparseDelegatePropert
 	FType PropertyParam = { GUE.Scopes.CoreUObject, FTypenameId(GUE.Names.NameType(Property->GetFName())) };
 	FType UniqueBindName = GUE.Names.MakeParametricType(MulticastSparseDelegate, {OwnerParam, PropertyParam});
 	FBindId Id = GUE.Names.IndexBindId(UniqueBindName);
-
-	GUE.Customs.BindStruct(Id, *Leak, GUE.Types.Get(GUE.Structs.MulticastDelegate));
+	
+	static FStructDeclarationPtr Declaration = DeclareMulticastDelegate();
+	GUE.Customs.BindStruct(Id, *Leak, FStructDeclarationPtr(Declaration), {});
 
 	return Id;
 }
 
 ////////////////////////////////////////////////////////////////////////////////////////////////
 
+FVerseFunctionBinding::FVerseFunctionBinding(TPropertySpecifier<1>& Spec)
+: MemberIds{GUE.Members.Value}
+{
+	Spec.Members[0] = SpecDynamicStruct;
+}
+
 void FVerseFunctionBinding::Save(FMemberBuilder& Dst, const FVerseFunction& Src, const FVerseFunction* Default, const FSaveContext& Ctx) const
 {
 	unimplemented();
@@ -831,6 +749,12 @@ bool FVerseFunctionBinding::Diff(const FVerseFunction& A, const FVerseFunction&
 	
 ////////////////////////////////////////////////////////////////////////////////////////////////
 
+FDynamicallyTypedValueBinding::FDynamicallyTypedValueBinding(TPropertySpecifier<1>& Spec)
+: MemberIds{GUE.Members.Value}
+{
+	Spec.Members[0] = SpecDynamicStruct;
+}
+
 void FDynamicallyTypedValueBinding::Save(FMemberBuilder& Dst, const Type& Src, const Type* Default, const FSaveContext& Ctx) const
 {
 	unimplemented();
@@ -849,6 +773,12 @@ bool FDynamicallyTypedValueBinding::Diff(const ::UE::FDynamicallyTypedValue& A,
 
 ////////////////////////////////////////////////////////////////////////////////////////////////
 
+FReferencePropertyBinding::FReferencePropertyBinding(TPropertySpecifier<1>& Spec)
+: MemberIds{GUE.Members.Value}
+{
+	Spec.Members[0] = SpecDynamicStruct;
+}
+
 void FReferencePropertyBinding::Save(FMemberBuilder& Dst, const FReferencePropertyValue& Src, const FReferencePropertyValue* Default, const FSaveContext& Ctx) const
 {
 	unimplemented();
@@ -909,14 +839,14 @@ struct FInterfaceBinding final : ICustomBinding
 class FInterfaceBindings
 {
 	const FType ScriptInterface;
-	const FStructDeclaration& Declaration;
+	FStructDeclarationPtr Declaration;
 
 	TMap<FType, FBindId> BoundClasses;
 
 public:
 	FInterfaceBindings()
 	: ScriptInterface(GUE.Names.Resolve(GUE.Structs.ScriptInterface))
-	, Declaration(GUE.Types.DeclareStruct(GUE.Structs.ScriptInterface, ScriptInterface, 0, {GUE.Members.Object}, EMemberPresence::RequireAll))
+	, Declaration(Declare({GUE.Structs.ScriptInterface, NoId, 0, EMemberPresence::RequireAll, {GUE.Members.Object}, {FMemberSpec(GUE.Structs.ObjectPtr)}}))
 	{}
 
 	FBindId Bind(FInterfaceProperty* Property)
@@ -933,7 +863,7 @@ public:
 
 		// Todo: Ownership / memory leak
 		ICustomBinding* Leak = new FInterfaceBinding(Property->InterfaceClass);
-		GUE.Customs.BindStruct(BindId, *Leak, Declaration);
+		GUE.Customs.BindStruct(BindId, *Leak, FStructDeclarationPtr(Declaration), {});
 
 		return BindId;
 	}
@@ -1043,7 +973,7 @@ static FEnumId DeclareEnum(UEnum* Enum)
 	checkSlow(Enumerators.Num() == Num || Enum->ContainsExistingMax() || Enum->HasMetaData(TEXT("Hidden"), Num));
 #endif
 
-	GUE.Types.DeclareEnum(Id, Type, Mode, Enumerators, EEnumAliases::Strip);
+	GUE.Enums.Declare(Id, Type, Mode, Enumerators, EEnumAliases::Strip);
 	return Id;
 }
 
@@ -1602,6 +1532,11 @@ struct TLeafRangeSerializer
 		return { FMemberType(DefaultRangeMax), InnerType, 1, EnumId, nullptr };
 	}
 
+	inline FMemberSpec SpecMember() const
+	{
+		return DefaultRangeOf(FMemberSpec(InnerType, FOptionalInnerId(EnumId)));
+	}
+
 	inline static EquivalentType SaveItem(const void* In, const FSaveContext&)
 	{
 		return Cast(In);
@@ -1637,6 +1572,12 @@ struct FStructRangeSerializer
 		return { FMemberType(DefaultRangeMax), InnerType, 1, FInnerId(SaveId), nullptr };
 	}
 
+	inline FMemberSpec SpecMember() const
+	{
+		// Lowercast is safe since FProperty doesnt support type-erasure
+		return DefaultRangeOf(LowerCast(SaveId));
+	}
+
 	inline FBuiltStruct* SaveItem(const void* In, const FSaveContext& Ctx) const
 	{
 		return SaveStruct(In, SaveId, Ctx);
@@ -1679,6 +1620,12 @@ struct FNestedRangeSerializer
 		return { FMemberType(DefaultRangeMax), InnerTypes[0], NumInners, InnermostSaveId, InnerTypes.GetData() };
 	}
 
+	inline FMemberSpec SpecMember() const
+	{
+		// InnermostSaveId is safe since FProperty doesnt support type-erasure
+		return FMemberSpec(InnerTypes, InnermostSaveId);
+	}
+
 	FBuiltRange* SaveItem(const void* In, const FSaveContext& Ctx) const
 	{
 		FRangeMemberBinding Member = { InnerBindTypes.GetData() + 1, InnerBindings.GetData(), NumInners - 1, InnermostSaveId, 0 };
@@ -1960,6 +1907,13 @@ struct TSetPropertyBinding : IItemRangeBinding, ICustomBinding
 		check(Inner.bHashable);
 	}
 
+	inline FStructDeclarationPtr DeclareCustom(FDeclId Id) const
+	{
+		FMemberId Members[] = { GUE.Members.Assign, GUE.Members.Remove, GUE.Members.Insert };
+		FMemberSpec RangeSpec = Range.SpecMember();
+		return Declare({Id, NoId, 0, EMemberPresence::AllowSparse, Members, {RangeSpec, RangeSpec, RangeSpec}});
+	}
+
 	inline SIZE_T GetStride() const requires (bLeaves)	{ return sizeof(TSetElement<LeafType>); }
 	inline SIZE_T GetStride() const						{ return Layout.Size; }
 
@@ -2212,26 +2166,26 @@ class FSetBindings
 	TMap<FParameterBinding, FBindId> Bindings;
 
 	template<EPropertyKind ItemKind>
-	void BindNew(FSetProperty* Property, FMemberBinding Elem, FBindId BindId, const FStructDeclaration& Declaration) 
+	void BindNew(FSetProperty* Property, FMemberBinding Elem, FBothStructId Both) 
 	{
 		// Todo: Ownership / memory leak
 		TSetPropertyBinding<ItemKind>* Leak = new TSetPropertyBinding<ItemKind>(Property, Elem);
-		GUE.Customs.BindStruct(BindId, *Leak, Declaration, {});
+		GUE.Customs.BindStruct(Both.BindId, *Leak, Leak->DeclareCustom(Both.DeclId), {});
 	}
 
-	void DeltaBindNew(FSetProperty* Property, FMemberBinding Elem, FBindId BindId, const FStructDeclaration& Declaration)
+	void DeltaBindNew(FSetProperty* Property, FMemberBinding Elem, FBothStructId Both)
 	{
 		switch (GetPropertyKind(Elem))
 		{
-		case EPropertyKind::Range:		BindNew<EPropertyKind::Range >(Property, Elem, BindId, Declaration); break;
-		case EPropertyKind::Struct:		BindNew<EPropertyKind::Struct>(Property, Elem, BindId, Declaration); break;
-		case EPropertyKind::Bool:		BindNew<EPropertyKind::Bool	 >(Property, Elem, BindId, Declaration); break;
-		case EPropertyKind::U8:			BindNew<EPropertyKind::U8	 >(Property, Elem, BindId, Declaration); break;
-		case EPropertyKind::U16:		BindNew<EPropertyKind::U16	 >(Property, Elem, BindId, Declaration); break;
-		case EPropertyKind::U32:		BindNew<EPropertyKind::U32	 >(Property, Elem, BindId, Declaration); break;
-		case EPropertyKind::U64:		BindNew<EPropertyKind::U64	 >(Property, Elem, BindId, Declaration); break;
-		case EPropertyKind::F32:		BindNew<EPropertyKind::F32	 >(Property, Elem, BindId, Declaration); break;
-		case EPropertyKind::F64:		BindNew<EPropertyKind::F64	 >(Property, Elem, BindId, Declaration); break;
+		case EPropertyKind::Range:		BindNew<EPropertyKind::Range >(Property, Elem, Both); break;
+		case EPropertyKind::Struct:		BindNew<EPropertyKind::Struct>(Property, Elem, Both); break;
+		case EPropertyKind::Bool:		BindNew<EPropertyKind::Bool	 >(Property, Elem, Both); break;
+		case EPropertyKind::U8:			BindNew<EPropertyKind::U8	 >(Property, Elem, Both); break;
+		case EPropertyKind::U16:		BindNew<EPropertyKind::U16	 >(Property, Elem, Both); break;
+		case EPropertyKind::U32:		BindNew<EPropertyKind::U32	 >(Property, Elem, Both); break;
+		case EPropertyKind::U64:		BindNew<EPropertyKind::U64	 >(Property, Elem, Both); break;
+		case EPropertyKind::F32:		BindNew<EPropertyKind::F32	 >(Property, Elem, Both); break;
+		case EPropertyKind::F64:		BindNew<EPropertyKind::F64	 >(Property, Elem, Both); break;
 		default:						check(false); break;
 		}
 	}
@@ -2253,11 +2207,8 @@ public:
 						 : BindType;
 		FBindId BindId = GUE.Names.IndexBindId(BindType);
 		FDeclId DeclId = Param.IsLowered() ? GUE.Names.IndexDeclId(DeclType) : LowerCast(BindId);
-		FMemberId Members[] = { GUE.Members.Assign, GUE.Members.Remove, GUE.Members.Insert };
 
-		// Todo: Ownership / memory leak
-		const FStructDeclaration& Declaration = GUE.Types.DeclareStruct(DeclId, DeclType, 0, Members, EMemberPresence::AllowSparse);
-		DeltaBindNew(Property, Elem, BindId, Declaration);
+		DeltaBindNew(Property, Elem, {BindId, DeclId});
 		return Bindings.Emplace(Elem, BindId);
 	}
 };
@@ -2387,6 +2338,14 @@ struct TMapPropertyCustomBinding : TMapPropertyItemBinding<FScriptMap, KeyKind,
 	, PairRange(Members.Pair)
 	{}
 
+	FStructDeclarationPtr Declare(FDeclId Id) const
+	{
+		FMemberId Members[] = { GUE.Members.Assign, GUE.Members.Remove, GUE.Members.Insert };
+		FMemberSpec KeysSpec = KeyRange.SpecMember();
+		FMemberSpec PairsSpec = PairRange.SpecMember();
+		return PlainProps::Declare({Id, NoId, 0, EMemberPresence::AllowSparse, Members, {PairsSpec, KeysSpec, PairsSpec}});
+	}
+
 	inline const void* GetValue(const void* Pair) const
 	{
 		return static_cast<const uint8*>(Pair) + Layout.ValueOffset;
@@ -2584,18 +2543,34 @@ class FMapBindings
 		}
 
 		// Index custom delta binding struct name
-		FParametricTypeId PairTypename = GUE.Names.Resolve(Members.Pair.InnermostSchema.Get().AsStruct()).Name.AsParametric();
-		TConstArrayView<FType> Params = GUE.Names.Resolve(PairTypename).GetParameters();
-		FType Type = { GUE.Scopes.Core, FTypenameId(GUE.Names.MakeParametricTypeId(GUE.Typenames.Map, Params)) };
-		FBindId Id = GUE.Names.IndexBindId(Type);
-		FMemberId MemberIds[] = { GUE.Members.Assign, GUE.Members.Remove, GUE.Members.Insert };
+		FBothStructId Both;
+		if (Members.Key.RangeBindings.Num() + Members.Value.RangeBindings.Num() > 0)
+		{
+			FBothType BothKey = Members.Key.IndexParameterName(GUE.Names);
+			FBothType BothValue = Members.Value.IndexParameterName(GUE.Names);
+			checkf(BothKey.IsLowered() || BothValue.IsLowered(), TEXT("Key or Value is range-bound and should be type-erased / lowered"));
+			FType BindParams[2] = {BothKey.BindType, BothValue.BindType};
+			FType DeclParams[2] = {BothKey.DeclType, BothValue.DeclType};
+			FType BindType = { GUE.Scopes.Core, FTypenameId(GUE.Names.MakeParametricTypeId(GUE.Typenames.Map, BindParams)) };
+			FType DeclType = { GUE.Scopes.Core, FTypenameId(GUE.Names.MakeParametricTypeId(GUE.Typenames.Map, DeclParams)) };
+			Both = { GUE.Names.IndexBindId(BindType), GUE.Names.IndexDeclId(DeclType) };
+		}
+		else
+		{
+			checkf(!Members.Key.IndexParameterName(GUE.Names).IsLowered(), TEXT("Only range-bound keys or values should be type-erased / lowered"));
+
+			FParametricTypeId PairTypename = GUE.Names.Resolve(Members.Pair.InnermostSchema.Get().AsStruct()).Name.AsParametric();
+			TConstArrayView<FType> Params = GUE.Names.Resolve(PairTypename).GetParameters();
+			FType Type = { GUE.Scopes.Core, FTypenameId(GUE.Names.MakeParametricTypeId(GUE.Typenames.Map, Params)) };
+			FDualStructId Dual{GUE.Names.IndexStruct(Type)};
+			Both = { Dual, Dual };
+		}
 
 		// Todo: Ownership / memory leak
-		const FStructDeclaration& Declaration = GUE.Types.DeclareStruct(LowerCast(Id), Type, 0, MemberIds, EMemberPresence::AllowSparse);
 		auto Out = new TMapPropertyCustomBinding<KeyKind, ValueKind>(Property, Members);
-		GUE.Customs.BindStruct(Id, *Out, Declaration, {});
+		GUE.Customs.BindStruct(Both.BindId, *Out, Out->Declare(Both.DeclId), {});
 
-		*OutCustomId = Id;
+		*OutCustomId = Both.BindId;
 		return Out;
 	}
 
@@ -2662,6 +2637,21 @@ static FMapBindings GMaps;
 
 ////////////////////////////////////////////////////////////////////////////////////////////////
 
+// @pre Binding.InnermostSchema isn't lowered
+static FMemberSpec ToSpec(FMemberBinding Binding)
+{
+	FMemberSpec Out = Binding.InnermostType.IsLeaf() 
+		? FMemberSpec(ToLeafType(Binding.InnermostType.AsLeaf()), ToOptionalEnum(Binding.InnermostSchema))
+		: FMemberSpec(Binding.InnermostSchema.Get().AsStructDeclId());
+
+	for (FRangeBinding RangeBinding : Binding.RangeBindings)
+	{
+		Out.RangeWrap(RangeBinding.GetSizeType());
+	}
+
+	return Out;
+}
+
 class FPairBindings
 {
 	struct FPair
@@ -2671,7 +2661,7 @@ class FPairBindings
 		inline bool operator==(const FPair& O) const { return EqSkipOffset(KV[0], O.KV[0]) && EqSkipOffset(KV[1], O.KV[1]); }
 	};
 
-	TMap<FPair, FBindId>	Bindings;
+	TMap<FPair, FBindId> Bindings;
 
 	inline FBindId BindImpl(FPair Pair)
 	{
@@ -2682,6 +2672,7 @@ class FPairBindings
 		}
 
 		// Index names, can be optimized by checking if KeyParam / BindParam IsLowered()
+		// or better checking if either one of them is a range
 		FBothType KeyParam = Pair.KV[0].IndexParameterName(GUE.Names);
 		FBothType ValueParam = Pair.KV[1].IndexParameterName(GUE.Names);
 		FType BindParams[2] = { KeyParam.BindType, ValueParam.BindType };
@@ -2690,11 +2681,12 @@ class FPairBindings
 		FType DeclType = { GUE.Scopes.Core, FTypenameId(GUE.Names.MakeParametricTypeId(GUE.Typenames.Pair, DeclParams)) };
 		FBindId BindId = GUE.Names.IndexBindId(BindType);
 		FDeclId DeclId = GUE.Names.IndexDeclId(DeclType);
+
 		FMemberId Members[2] = { GUE.Members.Key, GUE.Members.Value };
+		FMemberSpec Specs[2] = { ToSpec(Pair.KV[0]), ToSpec(Pair.KV[1]) };
 
 		// Todo: Ownership / memory leak
-		GUE.Types.DeclareStruct(DeclId, DeclType, 0, Members, EMemberPresence::RequireAll);
-		GUE.Schemas.BindStruct(BindId, DeclId, Pair.KV);
+		GUE.Schemas.BindStruct(BindId, Pair.KV, {DeclId, NoId, 0, EMemberPresence::RequireAll, Members, {Specs}});
 
 		Bindings.Emplace(Pair, BindId);
 
@@ -2948,24 +2940,60 @@ struct FStaticArrayBinding : IItemRangeBinding
 
 ////////////////////////////////////////////////////////////////////////////////////////////////
 
+// Serialize property-less UObject and UScriptStructs as their super class.
+//
+// E.g. FVector_NetQuantize10 is a pure runtime abstraction serialized as FVector.
+//		FAttenuationSubmixSendSettings is just a FSoundSubmixSendInfoBase but has a different 
+//		default constructor that matters in sparse delta serialization.
+//		UObjects are never instantiated during serialization so can be safely simplified.
+//
+// These heuristics might need more tuning
+static const UStruct* SkipEmptyBases(const UStruct* In)
+{
+	const UStruct* FirstOwner = In->PropertyLink ? In->PropertyLink->GetOwnerChecked<UStruct>() : In;
+	if (In != FirstOwner)
+	{
+		if (const UScriptStruct* Struct = Cast<const UScriptStruct>(In))
+		{
+			if (!(Struct->StructFlags & STRUCT_ZeroConstructor))
+			{
+				return In;
+			}
+		}
+
+		return FirstOwner;
+	}
+	return In;
+}
+
+////////////////////////////////////////////////////////////////////////////////////////////////
+
 class FPropertyBinder
 {
 public:
-	FPropertyBinder(FBindId Id) : Owner(Id) {}
+	FPropertyBinder(FBindId Id, EMemberPresence InOccupancy)
+	: Owner(Id)
+	, Occupancy(InOccupancy)
+	{}
 
-	void								AddMember(FMemberBinding Member)	{ Members.Add(Member); }
-	void								BindMember(FProperty* Property)		{ Members.Emplace(BindProperty(Property)); }
-	TConstArrayView<FMemberBinding>		GetMembers() const					{ return Members; }
+	void											BindSuper(FDeclId SuperId);
+	void											BindMember(FProperty* Property);
+	TConstArrayView<FMemberBinding>					GetMembers() const { return Members; }
+	FStructDeclarationPtr							Declare() const;	
 
 	// Only true for noexport UScriptStruct with STRUCT_Immutable | STRUCT_Atomic flags
-	bool								IsDense() { return GetOwnerOccupancy() == EMemberPresence::RequireAll; }
+	bool											IsDense() { return Occupancy == EMemberPresence::RequireAll; }
 
 private:
 	const FBindId									Owner;
+	FOptionalDeclId									Super;
 	TOptional<FScopeId>								OwnerScope;
-	TOptional<EMemberPresence>						OwnerOccupancy;
+	const EMemberPresence							Occupancy;
 	TPagedArray<FRangeBinding, 1024>				Ranges;
+	TArray<FMemberId, TInlineAllocator<64>>			Names;
 	TArray<FMemberBinding, TInlineAllocator<64>>	Members;
+	TArray<FMemberSpec, TInlineAllocator<64>>		Specs;
+
 	// BPVM only?
 	const FName										VerseFunctionProperty{"VerseFunctionProperty"};
 	const FName										VerseDynamicProperty{"VerseDynamicProperty"};
@@ -3011,16 +3039,6 @@ private:
 		return OwnerScope.GetValue();
 	}
 
-	EMemberPresence GetOwnerOccupancy()
-	{
-		if (!OwnerOccupancy)
-		{
-			OwnerOccupancy = GUE.Types.Get(LowerCast(Owner)).Occupancy;
-		}
-
-		return OwnerOccupancy.GetValue();
-	}
-
 	inline FMemberBinding BindAsRange(FProperty* Property, FRangeBinding RangeBinding, FMemberBinding Inner)
 	{
 		if (Inner.InnermostType.IsLeaf() && Inner.InnermostType.AsLeaf().Bind.Type == ELeafBindType::BitfieldBool)
@@ -3227,7 +3245,6 @@ private:
 		}
 
 		// Bind static array
-		EMemberPresence Occupancy = GetOwnerOccupancy();
 		uint32 TotalSize = static_cast<uint32>(Property->GetSize());
 		uint32 ElementSize = static_cast<uint32>(Property->GetElementSize());
 		uint32 ArrayDim = TotalSize / ElementSize;
@@ -3251,21 +3268,24 @@ private:
 		{
 			// Create struct binding to allow delta serialization
 			FType StaticArrayType = MakeStaticArrayTypename(Property->GetFName());
-			FBindId StaticArrayId = GUE.Names.IndexBindId(StaticArrayType);
+			FDualStructId StaticArrayId{GUE.Names.IndexStruct(StaticArrayType)};
 
-			// Todo: Ownership
-			TConstArrayView<FMemberId> Numerals = GUE.Numerals.MakeRange(IntCastChecked<uint16>(ArrayDim));
-			GUE.Types.DeclareNumeralStruct(LowerCast(StaticArrayId), StaticArrayType, Numerals, Occupancy);
-
-			TArray<FMemberBinding, TInlineAllocator<64>> Elements;
-			Elements.Init(Out, ArrayDim);
+			TArray<FMemberBinding, TInlineAllocator<64>> ElemBindings;			
+			ElemBindings.Init(Out, ArrayDim);
 			uint64 Offset = 0;
-			for (FMemberBinding& Element : Elements)
+			for (FMemberBinding& Element : ElemBindings)
 			{
 				Element.Offset = Offset;
 				Offset += ElementSize;
 			}
-			GUE.Schemas.BindStruct(StaticArrayId, LowerCast(StaticArrayId), Elements);
+			
+			TConstArrayView<FMemberId> Numerals = GUE.Numerals.MakeRange(IntCastChecked<uint16>(ArrayDim));
+			TArray<FMemberSpec, TInlineAllocator<64>> ElemSpecs;
+			ElemSpecs.Init(ToSpec(Out), ArrayDim);
+			FStructSpec Spec = {StaticArrayId, NoId, 0, EMemberPresence::AllowSparse, Numerals, { ElemSpecs }};
+
+			// Todo: Ownership
+			GUE.Schemas.BindStruct(StaticArrayId, ElemBindings, Spec);
 			
 			Out.InnermostType = DefaultStructBindType;
 			Out.InnermostSchema = FInnerId(StaticArrayId);
@@ -3276,6 +3296,56 @@ private:
 	}
 };
 
+void FPropertyBinder::BindSuper(FDeclId SuperId)
+{
+	check(!IsDense());
+
+	Super = SuperId;
+
+	FMemberBinding Member;
+	Member.InnermostType = SuperStructBindType;
+	Member.InnermostSchema = FInnerId(SuperId);
+	Members.Emplace(Member);
+}
+
+void FPropertyBinder::BindMember(FProperty* Property)
+{
+	Members.Emplace(BindProperty(Property));
+	Specs.Emplace(ToSpec(Members.Last()));
+	Names.Emplace(GUE.Names.NameMember(Property->GetFName()));
+}
+
+FStructDeclarationPtr FPropertyBinder::Declare() const
+{
+	return PlainProps::Declare({LowerCast(Owner), Super, 0, Occupancy, Names, {Specs}});
+}
+
+////////////////////////////////////////////////////////////////////////////////////////////////
+
+static bool ShouldBind(const UStruct* Struct);
+
+inline bool ShouldBind(const FProperty* Property)
+{
+	if (HasAny<CASTCLASS_FStructProperty>(Property->GetCastFlags()))
+	{
+		return ShouldBind(static_cast<const FStructProperty*>(Property)->Struct);
+	}
+
+	return true;
+}
+
+static bool ShouldBind(const UStruct* Struct)
+{
+	for (FProperty* Property = Struct->PropertyLink; Property; Property = Property->PropertyLinkNext)
+	{
+		if (ShouldBind(Property))
+		{
+			return true;
+		}
+	}
+	return false;
+}
+
 inline void BindMembers(FPropertyBinder& Out, const UStruct* Struct)
 {
 	for (FProperty* It = Struct->PropertyLink; It && It->GetOwner<UStruct>() == Struct; It = It->PropertyLinkNext)
@@ -3287,7 +3357,6 @@ inline void BindMembers(FPropertyBinder& Out, const UStruct* Struct)
 	}
 }
 
-// Must match DeclareSuperMembers
 void BindSuperMembers(FPropertyBinder& Out, const UStruct* Struct)
 {
 	if (const UStruct* Super = Struct->GetInheritanceSuper())
@@ -3300,6 +3369,12 @@ void BindSuperMembers(FPropertyBinder& Out, const UStruct* Struct)
 	}
 }
 
+const UStruct* GetSuperToBind(const UStruct* Struct)
+{
+	const UStruct* Super = Struct->GetInheritanceSuper();
+	return Super && ShouldBind(Super) ? SkipEmptyBases(Super) : nullptr;
+}
+
 void BindStruct(FBindId Id, const UStruct* Struct)
 {
 	if (GUE.Customs.FindStruct(Id))
@@ -3307,30 +3382,29 @@ void BindStruct(FBindId Id, const UStruct* Struct)
 		return;
 	}
 
-	FPropertyBinder Binder(Id);
-	if (Struct->GetInheritanceSuper())
+	FPropertyBinder Binder(Id, GetOccupancy(Struct));
+	if (const UStruct* Super = GetSuperToBind(Struct))
 	{
 		if (Binder.IsDense())
 		{
-			BindSuperMembers(/* out */ Binder, Struct);
+			// Flatten inheritance chain for dense structs
+			for (; Super; Super = GetSuperToBind(Super))
+			{
+				BindMembers(/* out */ Binder, Super);	
+			}
 		}
 		else
 		{
-			const FStructDeclaration& Declared = GUE.Types.Get(LowerCast(Id));
-			if (Declared.Super)
-			{
-				FMemberBinding Member;
-				Member.InnermostType = SuperStructBindType;
-				Member.InnermostSchema = FInnerId(Declared.Super.Get());
-				Binder.AddMember(Member);
-			}
+			FDeclId SuperId = GUE.Names.IndexDeclId(IndexType(Super));
+			Binder.BindSuper(SuperId);
 		}
 	}
+
 	BindMembers(/* out */ Binder, Struct);
 
-	GUE.Schemas.BindStruct(Id, LowerCast(Id), Binder.GetMembers());
+	GUE.Schemas.BindStruct(Id, Binder.GetMembers(), Binder.Declare());
 	
-	// Don't bind CDOs, object defaults are passed in from top and objects aren't owned by containers 
+	// Don't bind CDO defaults, object defaults are passed in from top and objects aren't owned by containers 
 	if (Struct->HasAnyCastFlags(CASTCLASS_UScriptStruct))
 	{
 		GUE.Defaults.Bind(Id, static_cast<const UScriptStruct*>(Struct));
@@ -3343,56 +3417,27 @@ static void BindInitialTypes()
 {
 	UE_LOGFMT(LogPlainPropsUObject, Display, "Binding types to PlainProps schemas...");
 
-	// Declare all UScriptStruct/UClass/UFunction/UEnum
-	static const FBindId SkipStructs[] = { GUE.Structs.VerseFunction };
-	TArray<FBindId> Ids;
-	for (TObjectIterator<UField> It; It; ++It)
+	// Declare all UEnums
+	for (TObjectIterator<UEnum> It; It; ++It)
 	{
-		UField* Field = *It;
-		if (UStruct* Struct = Cast<UStruct>(Field))
-		{
-			FType Type = IndexType(Struct);
-			FStructId Id = GUE.Names.IndexStruct(Type);
-			if (!Algo::Find(SkipStructs, FBindId(Id)))
-			{
-				DeclareStruct(Struct, Type, FDeclId(Id));
-			}
-			Ids.Add(FBindId(Id));
-		}
-		else if (UEnum* Enum = Cast<UEnum>(Field))
-		{
-			DeclareEnum(Enum);
-		}
+		DeclareEnum(*It);
 	}
 
 	// Bind all UScriptStruct/UClass/UFunction
-	const FBindId* IdIt = &Ids[0];
-	for (TObjectIterator<UStruct> It; It; ++It, ++IdIt)
+	static const FBindId SkipStructs[] = { GUE.Structs.VerseFunction };
+	for (TObjectIterator<UStruct> It; It; ++It)
 	{
-		if (!Algo::Find(SkipStructs, *IdIt))
+		const UStruct* Struct = *It;
+		FBindId Id = GUE.Names.IndexBindId(IndexType(Struct));
+		if (!Algo::Find(SkipStructs, Id))
 		{
-			BindStruct(*IdIt, *It);
+			BindStruct(Id, Struct);
 		}
 	}
-	check(IdIt == &Ids[0] + Ids.Num());
-}
-
-static void DeclareProperty(FDeclId Id, TConstArrayView<FMemberId> Members, EMemberPresence Occupancy)
-{
-	GUE.Types.DeclareStruct(Id,	GUE.Names.Resolve(Id), 0, Members, Occupancy);
 }
 
 static void InitBatchedProperties()
 {
-	// All property batch types share declaration, but some members could be exclusive to a specific type e.g. FMemoryPropertyBatch
-	DeclareProperty(GUE.Structs.Name,			{ GUE.Members.Id }, EMemberPresence::RequireAll);
-	DeclareProperty(GUE.Structs.Text,			{ GUE.Members.Id }, EMemberPresence::RequireAll);
-	DeclareProperty(GUE.Structs.ClassPtr,		{ GUE.Members.Id }, EMemberPresence::RequireAll);
-	DeclareProperty(GUE.Structs.ObjectPtr,		{ GUE.Members.Id }, EMemberPresence::RequireAll);
-	DeclareProperty(GUE.Structs.WeakObjectPtr,	{ GUE.Members.Id }, EMemberPresence::RequireAll);
-	DeclareProperty(GUE.Structs.SoftObjectPtr,	{ GUE.Members.Id }, EMemberPresence::RequireAll);
-	DeclareProperty(GUE.Structs.LazyObjectPtr,	{ GUE.Members.Id }, EMemberPresence::RequireAll);
-
 	GUE.Defaults.BindZeroes(GUE.Structs.Name, sizeof(FName), alignof(FName));
 	GUE.Defaults.BindStatic(GUE.Structs.Text, &FText::GetEmpty());
 	GUE.Defaults.BindZeroes(GUE.Structs.ClassPtr, sizeof(TSubclassOf<UClass>), alignof(TSubclassOf<UClass>));
@@ -3406,6 +3451,11 @@ struct FMemoryPropertyBatch
 {
 	TArray<FText> Texts; // Tricky to serialize intrusively
 
+	static FSoleMemberSpec Spec(FName*, FDeclId Id)
+	{
+		return {Id, GUE.Members.Id, Specify<uint64>()};
+	}
+
 	static void Save(FMemberBuilder& Out, FName In, const FSaveContext&)
 	{
 		Out.Add(GUE.Members.Id, FSensitiveName(In).ToUnstableInt());
@@ -3416,6 +3466,11 @@ struct FMemoryPropertyBatch
 		Out = FSensitiveName::FromUnstableInt(LoadSole<FSensitiveName::IntType>(In)).ToName();
 	}
 
+	static FSoleMemberSpec Spec(FText*, FDeclId Id)
+	{
+		return {Id, GUE.Members.Id, Specify<int32>()};
+	}
+
 	void Save(FMemberBuilder& Out, const FText& In, const FSaveContext&)
 	{
 		Out.Add(GUE.Members.Id, Texts.Num());
@@ -3427,6 +3482,11 @@ struct FMemoryPropertyBatch
 		Out = Texts[LoadSole<int32>(In)];
 	}
 
+	static FSoleMemberSpec Spec(FObjectHandle*, FDeclId Id)
+	{
+		return {Id, GUE.Members.Id, Specify<uint64>()};
+	}
+
 	static void Save(FMemberBuilder& Out, FObjectHandle In, const FSaveContext&)
 	{
 		static_assert(sizeof(In) == sizeof(uint64));
@@ -3438,6 +3498,11 @@ struct FMemoryPropertyBatch
 		LoadSole<uint64>(&Out, In);
 	}
 
+	static FSoleMemberSpec Spec(FWeakObjectPtr*, FDeclId Id)
+	{
+		return {Id, GUE.Members.Id, Specify<uint64>()};
+	}
+
 	static void Save(FMemberBuilder& Out, const FWeakObjectPtr& In, const FSaveContext&)
 	{
 		// Save ObjectSerialNumber + ObjectIndex a single uint64
@@ -3450,6 +3515,11 @@ struct FMemoryPropertyBatch
 		LoadSole<uint64>(&Out, In);
 	}
 
+	static FSoleMemberSpec Spec(FSoftObjectPtr*, FDeclId Id)
+	{
+		return {Id, GUE.Members.Id, FMemberSpec(GUE.Structs.SoftObjectPath)};
+	}
+
 	static void Save(FMemberBuilder& Out, const FSoftObjectPtr& In, const FSaveContext& Ctx)
 	{
 		FBuiltStruct* SoftPath = SaveStruct(&In.GetUniqueID(), GUE.Structs.SoftObjectPath, Ctx);
@@ -3461,7 +3531,12 @@ struct FMemoryPropertyBatch
 		Out.ResetWeakPtr();
 		LoadSoleStruct(&Out.GetUniqueID(), In);
 	}
-
+	
+	static FSoleMemberSpec Spec(FLazyObjectPtr*, FDeclId Id)
+	{
+		return {Id, GUE.Members.Id, FMemberSpec(GUE.Structs.Guid)};
+	}
+	
 	static void Save(FMemberBuilder& Out, const FLazyObjectPtr& In, const FSaveContext& Ctx)
 	{
 		FBuiltStruct* Guid = SaveStruct(&In.GetUniqueID(), GUE.Structs.Guid, Ctx);
@@ -3482,7 +3557,6 @@ inline bool DiffProperty(const FWeakObjectPtr& A, const FWeakObjectPtr& B) { ret
 inline bool DiffProperty(const FSoftObjectPtr& A, const FSoftObjectPtr& B) { return A != B; }
 inline bool DiffProperty(const FLazyObjectPtr& A, const FLazyObjectPtr& B) { return A != B; }
 
-
 template<class Type, class BatchType>
 struct TCustomPropertyBinding final : ICustomBinding
 {
@@ -3530,9 +3604,10 @@ struct TCustomPropertyBindings
 		Bind(GUE.Structs.LazyObjectPtr, LazyObjectPtr);
 	}
 
-	void Bind(FBindId Id, ICustomBinding& Binding)
+	template<class Type>
+	void Bind(FDualStructId Id, TCustomPropertyBinding<Type, BatchType>& Binding)
 	{
-		Overlay.BindStruct(Id, Binding, GUE.Types.Get(LowerCast(Id)));
+		Overlay.BindStruct(Id, Binding, Declare(BatchType::Spec((Type*)nullptr, Id)), {});
 	}
 
 	FCustomBindingsOverlay								Overlay;
@@ -3560,8 +3635,8 @@ class FBatchSaver
 {
 public:
 	FBatchSaver(FCustomBindings& Customs, int32 NumReserve)
-	: FlatCtx{{GUE.Types, GUE.Schemas, Customs}, Scratch, nullptr}
-	, DeltaCtx{{GUE.Types, GUE.Schemas, Customs}, Scratch, &GUE.Defaults}
+	: FlatCtx{{GUE.Schemas, Customs}, Scratch, nullptr}
+	, DeltaCtx{{GUE.Schemas, Customs}, Scratch, &GUE.Defaults}
 	{
 		SavedObjects.Reserve(NumReserve);
 	}
@@ -3577,8 +3652,8 @@ public:
 		ESchemaFormat Format = OutRuntimeIds ? ESchemaFormat::InMemoryNames : ESchemaFormat::StableNames;
 
 		// Build partial schemas
-		const FStructBindIds BindIds(GUE.Customs, GUE.Schemas);
-		FSchemasBuilder SchemaBuilders(GUE.Types, GUE.Names, BindIds, Scratch, Format);
+		const FBindDeclarations Types(GUE.Enums, FlatCtx.Customs, GUE.Schemas);
+		FSchemasBuilder SchemaBuilders(GUE.Names, Types, Scratch, Format);
 		for (const FSavedObject& Object : SavedObjects)
 		{
 			GDebugNoteObject = Object.Input;
@@ -3593,7 +3668,7 @@ public:
 			*OutRuntimeIds = ExtractRuntimeIds(Schemas);
 		}
 
-		FWriter Writer(GUE.Names, BindIds, Schemas, Format);
+		FWriter Writer(GUE.Names, Types, Schemas, Format);
 		TArray64<uint8> Out;
 
 		// Write out FNames when using StableNames
@@ -3662,8 +3737,6 @@ class FMemoryBatchLoader
 public:
 	FMemoryBatchLoader(const FCustomBindings& Customs, FMemoryView Data, TConstArrayView<FStructId> RuntimeIds)
 	{
-		//// Read ids
-		
 		// Read and mount schemas
 		FByteReader It(Data);
 		check(It.Grab<uint32>() == Magics[1]);
@@ -3687,7 +3760,7 @@ public:
 		check(!Objects.IsEmpty());
 
 		// Finally create load plans
-		Plans = CreateLoadPlans(Batch, GUE.Types, Customs, GUE.Schemas, RuntimeIds, ESchemaFormat::InMemoryNames);
+		Plans = CreateLoadPlans(Batch, Customs, GUE.Schemas, RuntimeIds, ESchemaFormat::InMemoryNames);
 	}
 
 	~FMemoryBatchLoader()
@@ -4194,7 +4267,7 @@ static int32 Roundtrip(ERoundtrip Options)
 
 		// Diff original vs PlainProps
 		UE_LOGFMT(LogPlainPropsUObject, Display, "Diffing UObjects roundtripped via PlainProps...");
-		FDiffContext DiffCtx = {{ GUE.Types, GUE.Schemas, Customs.Overlay }};
+		FDiffContext DiffCtx = {{ GUE.Schemas, Customs.Overlay }};
 		TUtf8StringBuilder<256> PpDiffs;
 		for (FInstance& Instance : Instances)
 		{
@@ -4238,7 +4311,7 @@ static int32 Roundtrip(ERoundtrip Options)
 	{
 		// Diff original vs TPS
 		UE_LOGFMT(LogPlainPropsUObject, Display, "Diffing UObjects roundtripped via TPS...");
-		FDiffContext DiffCtx = {{ GUE.Types, GUE.Schemas, Customs.Overlay }};
+		FDiffContext DiffCtx = {{ GUE.Schemas, Customs.Overlay }};
 		TUtf8StringBuilder<256> TpsDiffs;
 		TArray<int32> TpsDiffIdxs;
 		for (FInstance& Instance : Instances)
@@ -4293,7 +4366,7 @@ static int32 Roundtrip(ERoundtrip Options)
 
 		// Diff original vs UPS
 		UE_LOGFMT(LogPlainPropsUObject, Display, "Diffing UObjects roundtripped via UPS...");
-		FDiffContext DiffCtx = {{ GUE.Types, GUE.Schemas, Customs.Overlay }};
+		FDiffContext DiffCtx = {{ GUE.Schemas, Customs.Overlay }};
 		TUtf8StringBuilder<256> UpsDiffs;
 		TArray<int32> UpsDiffIdxs;
 		for (FInstance& Instance : Instances)
@@ -4414,5 +4487,4 @@ int32 UTestPlainPropsCommandlet::Main(const FString& Params)
 		EnumRemoveFlags(Options, ERoundtrip::TextMemory | ERoundtrip::TextStable);
 	}
 	return TestBindings(Options);
-}
-
+}
\ No newline at end of file
diff --git a/Engine/Plugins/Experimental/PlainPropsUObject/Source/Public/PlainPropsUObjectRuntime.h b/Engine/Plugins/Experimental/PlainPropsUObject/Source/Public/PlainPropsUObjectRuntime.h
--- a/Engine/Plugins/Experimental/PlainPropsUObject/Source/Public/PlainPropsUObjectRuntime.h
+++ b/Engine/Plugins/Experimental/PlainPropsUObject/Source/Public/PlainPropsUObjectRuntime.h
@@ -117,7 +117,7 @@ struct FGlobals
 	FGlobals();
 
 	TIdIndexer<FSensitiveName>	Names;
-	FDeclarations				Types;
+	FEnumDeclarations			Enums;
 	FSchemaBindings				Schemas;
 	FCustomBindingsBottom		Customs;
 	FDefaultStructs				Defaults;
@@ -126,6 +126,7 @@ struct FGlobals
 	FCommonTypenameIds			Typenames;
 	FCommonMemberIds			Members;
 	FNumeralGenerator			Numerals;
+	//Upgrade::FHistory			Upgrades;
 	FDebugIds					Debug;
 };
 PLAINPROPSUOBJECT_API extern FGlobals GUE;
@@ -146,7 +147,7 @@ struct FDefaultRuntime
 	using Ids = FRuntimeIds;
 	template<class T> using CustomBindings = TCustomBind<T>;
 
-	static FDeclarations&		GetTypes()			{ return GUE.Types; }
+	static FEnumDeclarations&	GetEnums()			{ return GUE.Enums; }
 	static FSchemaBindings&		GetSchemas()		{ return GUE.Schemas; }
 	static FCustomBindings&		GetCustoms()		{ return GUE.Customs; }
 	static IDefaultStructs*		GetDefaults()		{ return nullptr; }
@@ -159,6 +160,8 @@ struct FDeltaRuntime : FDefaultRuntime
 	static IDefaultStructs*		GetDefaults()		{ return &GUE.Defaults; }
 };
 
+template<uint32 N>
+using TPropertySpecifier = TCustomSpecifier<FRuntimeIds, N>;
 
 //////////////////////////////////////////////////////////////////////////
 
@@ -167,7 +170,7 @@ struct FFieldPathBinding : ICustomBinding
 	using Type = FFieldPath;
 	const FMemberId MemberIds[2];
 
-	FFieldPathBinding(TCustomInit<FRuntimeIds>) : MemberIds{GUE.Members.Path, GUE.Members.Owner} {}
+	FFieldPathBinding(TPropertySpecifier<2>& Spec);
 
 	void			Save(FMemberBuilder& Dst, const Type& Src, const Type* Default, const FSaveContext&) const;
 	void			Load(Type& Dst, FStructLoadView Src, ECustomLoadMethod Method) const;
@@ -179,8 +182,8 @@ struct FDelegateBinding : ICustomBinding
 	using Type = FScriptDelegate;
 	const FMemberId MemberIds[2];
 
-	FDelegateBinding(TCustomInit<FRuntimeIds>) : MemberIds{GUE.Members.Object, GUE.Members.Function} {}
-
+	FDelegateBinding(TPropertySpecifier<2>& Spec);
+	
 	void			Save(FMemberBuilder& Dst, const Type& Src, const Type* Default, const FSaveContext&) const;
 	void			Load(Type& Dst, FStructLoadView Src, ECustomLoadMethod Method) const;
 	static bool		Diff(const Type& A, const Type& B, const FBindContext&);
@@ -191,8 +194,8 @@ struct FMulticastInlineDelegateBinding : ICustomBinding
 	using Type = FMulticastScriptDelegate;
 	const FMemberId MemberIds[1];
 
-	FMulticastInlineDelegateBinding(TCustomInit<FRuntimeIds>) : MemberIds{GUE.Members.Invocations} {}
-
+	FMulticastInlineDelegateBinding(TPropertySpecifier<1>& Spec);
+	
 	void			Save(FMemberBuilder& Dst, const Type& Src, const Type* Default, const FSaveContext&) const;
 	void			Load(Type& Dst, FStructLoadView Src, ECustomLoadMethod Method) const;
 	static bool		Diff(const Type& A, const Type& B, const FBindContext&);
@@ -205,8 +208,8 @@ struct FVerseFunctionBinding : ICustomBinding
 	using Type =  FVerseFunction;
 	const FMemberId MemberIds[1];
 
-	FVerseFunctionBinding(TCustomInit<FRuntimeIds>) : MemberIds{GUE.Members.Value} {}
-
+	FVerseFunctionBinding(TPropertySpecifier<1>& Spec);
+	
 	void			Save(FMemberBuilder& Dst, const Type& Src, const Type* Default, const FSaveContext&) const;
 	void			Load(Type& Dst, FStructLoadView Src, ECustomLoadMethod Method) const;
 	static bool		Diff(const Type& A, const Type& B, const FBindContext&);
@@ -217,21 +220,20 @@ struct FDynamicallyTypedValueBinding : ICustomBinding
 	using Type = ::UE::FDynamicallyTypedValue;
 	const FMemberId MemberIds[1];
 
-	FDynamicallyTypedValueBinding(TCustomInit<FRuntimeIds>) : MemberIds{GUE.Members.Value} {}
-
+	FDynamicallyTypedValueBinding(TPropertySpecifier<1>& Spec);
+	
 	void			Save(FMemberBuilder& Dst, const Type& Src, const Type* Default, const FSaveContext&) const;
 	void			Load(Type& Dst, FStructLoadView Src, ECustomLoadMethod Method) const;
 	static bool		Diff(const Type& A, const Type& B, const FBindContext&);
 };
 
-
 struct FReferencePropertyBinding : ICustomBinding
 {
 	using Type = FReferencePropertyValue;
 	const FMemberId MemberIds[1];
 
-	FReferencePropertyBinding(TCustomInit<FRuntimeIds>) : MemberIds{GUE.Members.Value} {}
-
+	FReferencePropertyBinding(TPropertySpecifier<1>& Spec);
+	
 	void			Save(FMemberBuilder& Dst, const Type& Src, const Type* Default, const FSaveContext&) const;
 	void			Load(Type& Dst, FStructLoadView Src, ECustomLoadMethod Method) const;
 	static bool		Diff(const Type& A, const Type& B, const FBindContext&);
