PlainProps: WIP: Save and load source packages to and from memory using FLinkerLoad and FLinkerSave (not working yet).
The existing commandlet and unit tests still run.
Enable new logic by flipping PP_WITH_ENGINE_HOOKS in PlainPropsUObjectBindings.cpp and unshelving:
45229924: PlainProps: Minimalistic engine hooks for roundtripping source data packages using LinkerSave/LinkerLoad

Flow:
Change PlainPropsUObject startup phase to "PreLoadingScreen" to support most of the engine early loading phases (needs better coordinatation with UObjectBaseInit/InitAsyncThread in the future).
Refactor TestBindings into FPlainPropsBindingScope and handle setup/teardown in StartupModule/ShutdownModule.

Misc fixes:
Add FEnumDeclarations::Find(FEnumId Id) API to allow for declaring types lazily/on-demand.
Add temporary SaveStructWithSuper API to support delta serialization using a base object, needs to be hardened to guarantee super type correctness.
Add missing inner bindings for FMapBindings.
Add DumpUObjectToJsonFile to dump diffing UObjects to json files for easier introspection.
Add specialized DiffMap for FDiffContext diff the whole Pair to generate the full nested diff path for diffing values (handy for diffs in maps of dynamic blueprint types).

Commandlet tweaks:
Tweak FBatchSaver::Save to allow delta serialization using a base object.
Tweak FBatchSaver to write stable names using the existing SaveNameBatch API.
Add FPersistentBatchLoader using the existing LoadNameBatch batch API.
Add FLinkerPropertyBatch to save/load persistent custom CoreUObject bindings.
Add FPlainPropsLinkerSandbox to fake save/load persistent source packages (in-memory only).
For the package save test skip binding of transient and deprecated properties to match the linker tables.

#rnx

diff --git a/Engine/Plugins/Experimental/PlainProps/Source/Private/PlainPropsBind.cpp b/Engine/Plugins/Experimental/PlainProps/Source/Private/PlainPropsBind.cpp
--- a/Engine/Plugins/Experimental/PlainProps/Source/Private/PlainPropsBind.cpp
+++ b/Engine/Plugins/Experimental/PlainProps/Source/Private/PlainPropsBind.cpp
@@ -656,7 +656,7 @@ const FStructDeclaration& FBindContext::GetDeclaration(FBindId Id) const
 
 const FEnumDeclaration* FBindDeclarations::Find(FEnumId Id) const
 {
-	return &Enums.Get(Id);
+	return Enums.Find(Id);
 }
 
 const FStructDeclaration* FBindDeclarations::Find(FStructId Id) const
@@ -812,4 +812,4 @@ void DestroyTranslatedSchemas(const FSchemaBatch* Schemas)
 	FMemory::Free(const_cast<FSchemaBatch*>(Schemas));
 }
 
-} // namespace PlainProps
\ No newline at end of file
+} // namespace PlainProps
diff --git a/Engine/Plugins/Experimental/PlainProps/Source/Private/PlainPropsInternalParse.h b/Engine/Plugins/Experimental/PlainProps/Source/Private/PlainPropsInternalParse.h
--- a/Engine/Plugins/Experimental/PlainProps/Source/Private/PlainPropsInternalParse.h
+++ b/Engine/Plugins/Experimental/PlainProps/Source/Private/PlainPropsInternalParse.h
@@ -178,7 +178,7 @@ struct FParseDeclarations final : IDeclarations
 	const FStructDeclaration&			Get(FDeclId Id) { return *Structs[Id.Idx]; }
 
 	virtual FDeclId						Lower(FBindId Id) const override final;
-	virtual const FEnumDeclaration*		Find(FEnumId Id) const override { return &Enums.Get(Id); }
+	virtual const FEnumDeclaration*		Find(FEnumId Id) const override { return Enums.Find(Id); }
 	virtual const FStructDeclaration*	Find(FStructId Id) const override;
 
 	FEnumDeclarations					Enums;
@@ -214,4 +214,4 @@ private:
 	int64											ObjectsIdx = -1;
 };
 
-} // namespace PlainProps
\ No newline at end of file
+} // namespace PlainProps
diff --git a/Engine/Plugins/Experimental/PlainProps/Source/Private/PlainPropsLoad.cpp b/Engine/Plugins/Experimental/PlainProps/Source/Private/PlainPropsLoad.cpp
--- a/Engine/Plugins/Experimental/PlainProps/Source/Private/PlainPropsLoad.cpp
+++ b/Engine/Plugins/Experimental/PlainProps/Source/Private/PlainPropsLoad.cpp
@@ -361,7 +361,7 @@ static void CreateSubsetBindingWithoutEnumIds(const FStructSchema& From, const F
 	// ...first the unnamed super member...
 	if (From.Inheritance != ESuper::No)
 	{
-		FBindId BindId = ToIt.GrabSuper();
+		FBindId DeclBindId = ToIt.GrabSuper();
 		if (UsesSuper(From.Inheritance))
 		{
 			FStructType FromType = From.Footer[0].AsStruct(); // To.Members[0].AsStruct().IsDynamic isn't set
@@ -371,7 +371,18 @@ static void CreateSubsetBindingWithoutEnumIds(const FStructSchema& From, const F
 			if (!FromType.IsDynamic)
 			{
 				FStructSchemaId ReadId = static_cast<FStructSchemaId>(*InnerSchemaIt);
-				FStructSchemaId LoadId = LoadIds.BindLoadId(ReadId, BindId);
+				FStructSchemaId LoadId;
+				if (From.Inheritance == ESuper::Reused)
+				{
+					LoadId = LoadIds.BindLoadId(ReadId, DeclBindId);
+				}
+				else
+				{
+					check(From.Inheritance == ESuper::Used);
+					// TODO: Handle type erasure in the used super scenario too,
+					// i.e. when reading a base of the declared super, needs UsedBindId.
+					LoadId = ReadId;
+				}
 				Footer.AddInnerSchema(LoadId);
 				++InnerSchemaIt;
 			}
@@ -1350,4 +1361,4 @@ void LoadSoleStruct(void* Dst, FStructLoadView Src)
 	//FLoadStructPlan Plan = Src.Schema.Batch[Src.Schema.LoadId];
 }
 
-} // namespace PlainProps
\ No newline at end of file
+} // namespace PlainProps
diff --git a/Engine/Plugins/Experimental/PlainProps/Source/Private/PlainPropsSave.cpp b/Engine/Plugins/Experimental/PlainProps/Source/Private/PlainPropsSave.cpp
--- a/Engine/Plugins/Experimental/PlainProps/Source/Private/PlainPropsSave.cpp
+++ b/Engine/Plugins/Experimental/PlainProps/Source/Private/PlainPropsSave.cpp
@@ -302,17 +302,29 @@ static void SaveMember(FMemberBuilder& Out, const void* Struct, FMemberId Name,
 	Out.AddStruct(Name, Member.Id, SaveStruct(At(Struct, Member.Offset), Member.Id, Ctx));
 }
 
-inline void SaveAllMembers(FMemberBuilder& Out, const void* Struct, const FSchemaBinding& Schema, const FStructDeclaration& Declaration, const FSaveContext& Ctx)
+// Allow a super that has been prebuilt with SaveStructDelta
+inline void SaveAllMembers(FMemberBuilder& Out, const void* Struct, const FSchemaBinding& Schema, const FStructDeclaration& Declaration, const FSaveContext& Ctx, FBuiltStruct* BuiltSuper = nullptr)
 {
 	FMemberVisitor It(Schema);
 	if (Declaration.Super)
 	{
 		FBindId SuperId = It.GrabSuper();
 		checkSlow(SuperId == ToOptionalStruct(Declaration.Super));
-		const FStructDeclaration* SuperDecl;
-		const FSchemaBinding& SuperSchema = Ctx.Schemas.GetStruct(SuperId, /* out */ SuperDecl);
-		SaveAllMembers(Out, Struct, SuperSchema, *SuperDecl, Ctx);
-		Out.BuildSuperStruct(Ctx.Scratch, *SuperDecl, Ctx.Schemas.GetDebug());
+		if (BuiltSuper)
+		{
+			Out.AddSuperStruct(SuperId, BuiltSuper);
+		}
+		else
+		{
+			const FStructDeclaration* SuperDecl;
+			const FSchemaBinding& SuperSchema = Ctx.Schemas.GetStruct(SuperId, /* out */ SuperDecl);
+			SaveAllMembers(Out, Struct, SuperSchema, *SuperDecl, Ctx);
+			Out.BuildSuperStruct(Ctx.Scratch, *SuperDecl, Ctx.Schemas.GetDebug());
+		}
+	}
+	else
+	{
+		checkSlow(!BuiltSuper);
 	}
 
 	for (FMemberId Name : Declaration.GetMemberOrder())
@@ -344,6 +356,20 @@ FBuiltStruct* SaveStruct(const void* Struct, FBindId Id, const FSaveContext& Ctx
 	return Out.BuildAndReset(Ctx.Scratch, *Declaration, Ctx.Schemas.GetDebug());
 }
 
+
+FBuiltStruct* SaveStructWithSuper(const void* Struct, FBuiltStruct* BuiltSuper, FBindId Id, const FSaveContext& Ctx)
+{
+	check(BuiltSuper);
+	const FStructDeclaration* Declaration = nullptr;
+	check(!Ctx.Customs.FindStructToSave(Id, /* out */ Declaration));
+	FMemberBuilder Out;
+
+	const FSchemaBinding& Schema = Ctx.Schemas.GetStruct(Id, /* out */ Declaration);
+	check(Declaration->Super);
+	SaveAllMembers(Out, Struct, Schema, *Declaration, Ctx, BuiltSuper);
+	return Out.BuildAndReset(Ctx.Scratch, *Declaration, Ctx.Schemas.GetDebug());
+}
+
 ////////////////////////////////////////////////////////////////////////////////////////////////
 
 static bool DiffItem(const uint8* A, const uint8* B, const FSaveContext& Ctx, FRangeMemberBinding Member)
diff --git a/Engine/Plugins/Experimental/PlainProps/Source/Public/PlainPropsDeclare.h b/Engine/Plugins/Experimental/PlainProps/Source/Public/PlainPropsDeclare.h
--- a/Engine/Plugins/Experimental/PlainProps/Source/Public/PlainPropsDeclare.h
+++ b/Engine/Plugins/Experimental/PlainProps/Source/Public/PlainPropsDeclare.h
@@ -50,7 +50,9 @@ public:
 
 	PLAINPROPS_API const FEnumDeclaration&	Declare(FEnumId Id, FType Type, EEnumMode Mode, TConstArrayView<FEnumerator> Enumerators, EEnumAliases Policy);
 	void									Drop(FEnumId Id)				{ Check(Id); Declarations[Id.Idx].Reset(); }
-	const FEnumDeclaration&					Get(FEnumId Id) const			{ Check(Id); return *Declarations[Id.Idx]; }
+	const FEnumDeclaration&					Get(FEnumId Id) const			{ Check(Id); return *Declarations[Id.Idx].Get(); }
+	const FEnumDeclaration*					Find(FEnumId Id) const			{ return Id.Idx < (uint32)Declarations.Num() ? Declarations[Id.Idx].Get() : nullptr; }
+
 	FDebugIds								GetDebug() const				{ return Debug; }
 
 protected:
@@ -107,4 +109,4 @@ struct IDeclarations
 	virtual const FStructDeclaration*	Find(FStructId Id) const = 0;
 };
 
-} // namespace PlainProps
\ No newline at end of file
+} // namespace PlainProps
diff --git a/Engine/Plugins/Experimental/PlainProps/Source/Public/PlainPropsSave.h b/Engine/Plugins/Experimental/PlainProps/Source/Public/PlainPropsSave.h
--- a/Engine/Plugins/Experimental/PlainProps/Source/Public/PlainPropsSave.h
+++ b/Engine/Plugins/Experimental/PlainProps/Source/Public/PlainPropsSave.h
@@ -7,6 +7,8 @@
 namespace PlainProps 
 {
 
+struct FBuiltStruct;
+
 struct FSaveContext : FBindContext
 {
 	FScratchAllocator&			Scratch;
@@ -22,9 +24,11 @@ FSaveContext MakeSaveContext(FScratchAllocator& Scratch)
 ////////////////////////////////////////////////////////////////////////////////////////////////
 
 [[nodiscard]] PLAINPROPS_API FBuiltStruct*	SaveStruct(const void* Struct, FBindId BindId, const FSaveContext& Context);
+// todo: temp api for delta serialization with a superclass default object
+[[nodiscard]] PLAINPROPS_API FBuiltStruct*	SaveStructWithSuper(const void* Struct, FBuiltStruct* BuiltSuper, FBindId BindId, const FSaveContext& Context);
 [[nodiscard]] PLAINPROPS_API FBuiltStruct*	SaveStructDelta(const void* Struct, const void* Default, FBindId BindId, const FSaveContext& Context);
 [[nodiscard]] PLAINPROPS_API FBuiltStruct*	SaveStructDeltaIfDiff(const void* Struct, const void* Default, FBindId BindId, const FSaveContext& Context);
 [[nodiscard]] PLAINPROPS_API FBuiltRange*	SaveRange(const void* Range, FRangeMemberBinding Member, const FSaveContext& Ctx);
 [[nodiscard]] PLAINPROPS_API FBuiltRange*	SaveLeafRange(const void* Range, const ILeafRangeBinding& Binding, FUnpackedLeafType Leaf, const FSaveContext& Ctx);
 
-} // namespace PlainProps
\ No newline at end of file
+} // namespace PlainProps
diff --git a/Engine/Plugins/Experimental/PlainProps/Source/Tests/PlainPropsTests.cpp b/Engine/Plugins/Experimental/PlainProps/Source/Tests/PlainPropsTests.cpp
--- a/Engine/Plugins/Experimental/PlainProps/Source/Tests/PlainPropsTests.cpp
+++ b/Engine/Plugins/Experimental/PlainProps/Source/Tests/PlainPropsTests.cpp
@@ -174,7 +174,7 @@ private:
 
 	TArray<char> GetNameData() const;
 
-	virtual const FEnumDeclaration*		Find(FEnumId Id) const override { return &EnumDeclarations.Get(Id); }
+	virtual const FEnumDeclaration*		Find(FEnumId Id) const override { return EnumDeclarations.Find(Id); }
 	virtual const FStructDeclaration*	Find(FStructId Id) const override { return StructDeclarations.FindChecked(FDeclId(Id)); }
 	virtual FDeclId						Lower(FBindId Id) const override
 	{
diff --git a/Engine/Plugins/Experimental/PlainPropsUObject/PlainPropsUObject.uplugin b/Engine/Plugins/Experimental/PlainPropsUObject/PlainPropsUObject.uplugin
--- a/Engine/Plugins/Experimental/PlainPropsUObject/PlainPropsUObject.uplugin
+++ b/Engine/Plugins/Experimental/PlainPropsUObject/PlainPropsUObject.uplugin
@@ -14,7 +14,8 @@
 		{
 			"Name": "PlainPropsUObject",
 			"Type": "Runtime",
-			"PlatformAllowList": [ "Win64" ]
+			"LoadingPhase": "PreLoadingScreen",
+			"PlatformAllowList": [ "Win64" ],
 		}
 	],
 	"Plugins": [
@@ -23,4 +24,4 @@
 			"Enabled": true
 		}
 	]
-}
\ No newline at end of file
+}
diff --git a/Engine/Plugins/Experimental/PlainPropsUObject/Source/PlainPropsUObject.Build.cs b/Engine/Plugins/Experimental/PlainPropsUObject/Source/PlainPropsUObject.Build.cs
--- a/Engine/Plugins/Experimental/PlainPropsUObject/Source/PlainPropsUObject.Build.cs
+++ b/Engine/Plugins/Experimental/PlainPropsUObject/Source/PlainPropsUObject.Build.cs
@@ -18,5 +18,15 @@ public class PlainPropsUObject : ModuleRules
 				"Engine",
 				"PlainProps"
 		});
+
+		if (Target.bBuildEditor == true)
+		{
+			PrivateDependencyModuleNames.AddRange(
+				new string[] {
+					"UnrealEd", 		// for GEditor->SavePackage
+					"JsonUtilities", 	// for FJsonObjectConverter::UStructToJsonObjectString
+				}
+			);
+		}
 	}
 }
diff --git a/Engine/Plugins/Experimental/PlainPropsUObject/Source/Private/PlainPropsUObjectBindings.cpp b/Engine/Plugins/Experimental/PlainPropsUObject/Source/Private/PlainPropsUObjectBindings.cpp
--- a/Engine/Plugins/Experimental/PlainPropsUObject/Source/Private/PlainPropsUObjectBindings.cpp
+++ b/Engine/Plugins/Experimental/PlainPropsUObject/Source/Private/PlainPropsUObjectBindings.cpp
@@ -7,6 +7,7 @@
 #include "PlainPropsParse.h"
 #include "PlainPropsPrint.h"
 #include "PlainPropsRead.h"
+#include "PlainPropsUObjectModule.h"
 #include "PlainPropsVisualize.h"
 #include "PlainPropsWrite.h"
 #include "Algo/Compare.h"
@@ -14,33 +15,52 @@
 #include "Containers/PagedArray.h"
 #include "HAL/FileManager.h"
 #include "Hash/xxhash.h"
+#include "JsonObjectConverter.h"
 #include "Logging/StructuredLog.h"
 #include "Math/PreciseFP.h"
 #include "Misc/AsciiSet.h"
 #include "Misc/CommandLine.h"
 #include "Misc/CoreDelegates.h"
 #include "Misc/DefinePrivateMemberPtr.h"
+#include "Misc/Filehelper.h"
 #include "Misc/Paths.h"
+#include "Misc/PackageName.h"
 #include "Serialization/Archive.h"
+#include "Serialization/BufferWriter.h"
+#include "Serialization/LargeMemoryReader.h"
 #include "Serialization/MemoryReader.h"
 #include "Serialization/MemoryWriter.h"
+#include "Serialization/UnversionedPropertySerializationInternal.h"
 #include "StructUtils/UserDefinedStruct.h"
 #include "Templates/UniquePtr.h"
 #include "UObject/AnsiStrProperty.h"
 #include "UObject/Class.h"
 #include "UObject/EnumProperty.h"
 #include "UObject/FieldPathProperty.h"
+#include "UObject/LinkerLoad.h"
+#include "UObject/LinkerSave.h"
+#include "UObject/NameBatchSerialization.h"
 #include "UObject/Object.h"
+#include "UObject/ObjectHandle.h"
+#include "UObject/ObjectResource.h"
 #include "UObject/PropertyOptional.h"
+#include "UObject/SavePackage.h"
 #include "UObject/StrProperty.h"
 #include "UObject/TextProperty.h"
 #include "UObject/Utf8StrProperty.h"
 #include "UObject/UnrealType.h"
 #include "UObject/UObjectIterator.h"
+#include "UObject/UObjectThreadContext.h"
 #include "UObject/VerseValueProperty.h"
 #include "UObject/VerseStringProperty.h"
 
 
+#if WITH_EDITOR
+#include "Editor.h"
+#endif
+
+#define PP_WITH_ENGINE_HOOKS 0
+
 using FUnicastScriptDelegate = TScriptDelegate<FNotThreadSafeNotCheckedDelegateMode>;
 using FMulticastInvocationList = TArray<FUnicastScriptDelegate>;
 using FMulticastInvocationView = TConstArrayView<FUnicastScriptDelegate>;
@@ -105,7 +125,11 @@ static FDefaultStruct* NewDefaultStruct(UScriptStruct::ICppStructOps& Ops)
 	uint32 Size = sizeof(FDefaultStruct) + Ops.GetSize();
 	FDefaultStruct Header = {Ops};
 	FDefaultStruct* Out = new (FMemory::MallocZeroed(Size)) FDefaultStruct(Header);
+	FUObjectThreadContext& ThreadContext = FUObjectThreadContext::Get();
+	ThreadContext.IsInConstructor++; // Hack: FWaterBrushActorDefaults is doing ConstructorHelpers::FObjectFinder<UCurveFloat>
 	Ops.Construct(Out->Instance);
+	ThreadContext.IsInConstructor--;
+	check(ThreadContext.IsInConstructor >= 0);
 	return Out;
 }
 
@@ -947,6 +971,11 @@ static FEnumId DeclareEnum(UEnum* Enum)
 {
 	FType Type = IndexType(Enum);
 	FEnumId Id = GUE.Names.IndexEnum(Type);
+	if (GUE.Enums.Find(Id))
+	{
+		return Id;
+	}
+
 	EEnumMode Mode = Enum->HasAnyEnumFlags(EEnumFlags::Flags) ? EEnumMode::Flag : EEnumMode::Flat;
 
 	// Skip _MAX and _All enumerators
@@ -1842,6 +1871,7 @@ void LoadSetDelta(const BindingType& Binding, ScriptType& Dst, FStructLoadView S
 	else if (Name == GUE.Members.Assign)
 	{
 		Binding.DestroyAll(Dst);
+		Dst.Empty(Dst.Max(), Binding.Layout);
 		Binding.AssignEmpty(Dst, Range);
 	}
 	else
@@ -1880,6 +1910,43 @@ inline bool DiffSet(const BindingType& Binding, const ScriptType& A, const Scrip
 	return false;
 }
 
+template<typename BindingType, typename ScriptType>
+bool DiffMap(const BindingType& Binding, const ScriptType& A, const ScriptType& B, FDiffContext& Ctx)
+{
+	if (A.NumUnchecked() != B.NumUnchecked())
+	{
+		return true;
+	}
+
+	if (A.NumUnchecked() > 0)
+	{
+		for (int32 IdxA = 0, MaxA = A.GetMaxIndex(); IdxA < MaxA; ++IdxA)
+		{
+			if (!A.IsValidIndex(IdxA))
+			{
+				continue;
+			}
+			const void* KeyA = A.GetData(IdxA, Binding.Layout);
+			int32 IdxB = Binding.FindKey(B, KeyA);
+			if (IdxB == INDEX_NONE)
+			{
+				// Todo: Add specific diff entry for missing key.
+				return true;
+			}
+			FBindId ValueId = Binding.PairRange.SaveId;
+			// Diff the whole Pair to generate the full nested diff path for diffing values
+			if (DiffStructs(/* GetValue */(KeyA), /* GetValue */(B.GetData(IdxB, Binding.Layout)), ValueId, Ctx))
+			{
+				check(!Binding.HasItem(B, A.GetData(IdxA, Binding.Layout)));
+				return true;
+			}
+			check(Binding.HasItem(B, A.GetData(IdxA, Binding.Layout)));
+		}
+	}
+
+	return false;
+}
+
 ////////////////////////////////////////////////////////////////////////////////////////////////
 
 template<EPropertyKind ElemKind>
@@ -2525,6 +2592,11 @@ struct TMapPropertyCustomBinding : TMapPropertyItemBinding<FScriptMap, KeyKind,
 	{
 		return DiffSet(*this, *static_cast<const FScriptMap*>(A), *static_cast<const FScriptMap*>(B));
 	}
+	
+	virtual bool DiffCustom(const void* A, const void* B, FDiffContext& Ctx) const override
+	{
+		return DiffMap(*this, *static_cast<const FScriptMap*>(A), *static_cast<const FScriptMap*>(B), Ctx);
+	}
 };
 
 ////////////////////////////////////////////////////////////////////////////////////////////////
@@ -2543,7 +2615,8 @@ class FMapBindings
 		}
 
 		// Index custom delta binding struct name
-		FBothStructId Both;
+		FBothStructId BothPair;
+		TArray<FInnerStruct, TInlineAllocator<3>> Inners;
 		if (Members.Key.RangeBindings.Num() + Members.Value.RangeBindings.Num() > 0)
 		{
 			FBothType BothKey = Members.Key.IndexParameterName(GUE.Names);
@@ -2553,24 +2626,39 @@ class FMapBindings
 			FType DeclParams[2] = {BothKey.DeclType, BothValue.DeclType};
 			FType BindType = { GUE.Scopes.Core, FTypenameId(GUE.Names.MakeParametricTypeId(GUE.Typenames.Map, BindParams)) };
 			FType DeclType = { GUE.Scopes.Core, FTypenameId(GUE.Names.MakeParametricTypeId(GUE.Typenames.Map, DeclParams)) };
-			Both = { GUE.Names.IndexBindId(BindType), GUE.Names.IndexDeclId(DeclType) };
+			BothPair = { GUE.Names.IndexBindId(BindType), GUE.Names.IndexDeclId(DeclType) };
+
+			FBindId PairId = Members.Pair.InnermostSchema.Get().AsStructBindId();
+			if (BothKey.IsLowered())
+			{
+				FBindId KeyId = GUE.Names.IndexBindId(BothKey.BindType);
+				Inners.Emplace(GUE.Members.Assign, PairId);
+				Inners.Emplace(GUE.Members.Remove, KeyId);
+				Inners.Emplace(GUE.Members.Insert, PairId);
+			}
+			else
+			{
+				Inners.Emplace(GUE.Members.Assign, PairId);
+				Inners.Emplace(GUE.Members.Insert, PairId);
+			}
 		}
 		else
 		{
 			checkf(!Members.Key.IndexParameterName(GUE.Names).IsLowered(), TEXT("Only range-bound keys or values should be type-erased / lowered"));
+			checkf(!Members.Value.IndexParameterName(GUE.Names).IsLowered(), TEXT("Only range-bound keys or values should be type-erased / lowered"));
 
 			FParametricTypeId PairTypename = GUE.Names.Resolve(Members.Pair.InnermostSchema.Get().AsStruct()).Name.AsParametric();
 			TConstArrayView<FType> Params = GUE.Names.Resolve(PairTypename).GetParameters();
 			FType Type = { GUE.Scopes.Core, FTypenameId(GUE.Names.MakeParametricTypeId(GUE.Typenames.Map, Params)) };
 			FDualStructId Dual{GUE.Names.IndexStruct(Type)};
-			Both = { Dual, Dual };
+			BothPair = { Dual, Dual };
 		}
 
 		// Todo: Ownership / memory leak
 		auto Out = new TMapPropertyCustomBinding<KeyKind, ValueKind>(Property, Members);
-		GUE.Customs.BindStruct(Both.BindId, *Out, Out->Declare(Both.DeclId), {});
+		GUE.Customs.BindStruct(BothPair.BindId, *Out, Out->Declare(BothPair.DeclId), Inners);
 
-		*OutCustomId = Both.BindId;
+		*OutCustomId = BothPair.BindId;
 		return Out;
 	}
 
@@ -3322,122 +3410,142 @@ FStructDeclarationPtr FPropertyBinder::Declare() const
 
 ////////////////////////////////////////////////////////////////////////////////////////////////
 
-static bool ShouldBind(const UStruct* Struct);
-
-inline bool ShouldBind(const FProperty* Property)
+struct FTypeBinder
 {
-	if (HasAny<CASTCLASS_FStructProperty>(Property->GetCastFlags()))
-	{
-		return ShouldBind(static_cast<const FStructProperty*>(Property)->Struct);
-	}
-
-	return true;
-}
+	EPropertyFlags ExcludedPropertyFlags = EPropertyFlags::CPF_None;
 
-static bool ShouldBind(const UStruct* Struct)
-{
-	for (FProperty* Property = Struct->PropertyLink; Property; Property = Property->PropertyLinkNext)
+	inline bool ShouldBind(const FProperty* Property)
 	{
-		if (ShouldBind(Property))
+		if (Property->PropertyFlags & ExcludedPropertyFlags)
 		{
-			return true;
+			return false;
 		}
-	}
-	return false;
-}
 
-inline void BindMembers(FPropertyBinder& Out, const UStruct* Struct)
-{
-	for (FProperty* It = Struct->PropertyLink; It && It->GetOwner<UStruct>() == Struct; It = It->PropertyLinkNext)
-	{
-		if (ShouldBind(It))
+		if (HasAny<CASTCLASS_FStructProperty>(Property->GetCastFlags()))
 		{
-			Out.BindMember(It);
+			return ShouldBind(static_cast<const FStructProperty*>(Property)->Struct);
 		}
+
+		return true;
 	}
-}
 
-void BindSuperMembers(FPropertyBinder& Out, const UStruct* Struct)
-{
-	if (const UStruct* Super = Struct->GetInheritanceSuper())
+	bool ShouldBind(const UStruct* Struct)
 	{
-		BindSuperMembers(Out, Super);
-		if (ShouldBind(Super))
+		for (FProperty* Property = Struct->PropertyLink; Property; Property = Property->PropertyLinkNext)
 		{
-			BindMembers(Out, Super);
+			if (ShouldBind(Property))
+			{
+				return true;
+			}
 		}
+		return false;
 	}
-}
-
-const UStruct* GetSuperToBind(const UStruct* Struct)
-{
-	const UStruct* Super = Struct->GetInheritanceSuper();
-	return Super && ShouldBind(Super) ? SkipEmptyBases(Super) : nullptr;
-}
 
-void BindStruct(FBindId Id, const UStruct* Struct)
-{
-	if (GUE.Customs.FindStruct(Id))
+	inline void BindMembers(FPropertyBinder& Out, const UStruct* Struct)
 	{
-		return;
+		for (FProperty* It = Struct->PropertyLink; It && It->GetOwner<UStruct>() == Struct; It = It->PropertyLinkNext)
+		{
+			if (ShouldBind(It))
+			{
+				Out.BindMember(It);
+			}
+		}
 	}
 
-	FPropertyBinder Binder(Id, GetOccupancy(Struct));
-	if (const UStruct* Super = GetSuperToBind(Struct))
+	void BindSuperMembers(FPropertyBinder& Out, const UStruct* Struct)
 	{
-		if (Binder.IsDense())
+		if (const UStruct* Super = Struct->GetInheritanceSuper())
 		{
-			// Flatten inheritance chain for dense structs
-			for (; Super; Super = GetSuperToBind(Super))
+			BindSuperMembers(Out, Super);
+			if (ShouldBind(Super))
 			{
-				BindMembers(/* out */ Binder, Super);	
+				BindMembers(Out, Super);
 			}
 		}
-		else
+	}
+
+	const UStruct* GetSuperToBind(const UStruct* Struct)
+	{
+		const UStruct* Super = Struct->GetInheritanceSuper();
+		while (Super)
 		{
-			FDeclId SuperId = GUE.Names.IndexDeclId(IndexType(Super));
-			Binder.BindSuper(SuperId);
+			if (ShouldBind(Super))
+			{
+				return SkipEmptyBases(Super);
+			}
+			Super = Super->GetInheritanceSuper();
 		}
+		return nullptr;
 	}
 
-	BindMembers(/* out */ Binder, Struct);
-
-	GUE.Schemas.BindStruct(Id, Binder.GetMembers(), Binder.Declare());
-	
-	// Don't bind CDO defaults, object defaults are passed in from top and objects aren't owned by containers 
-	if (Struct->HasAnyCastFlags(CASTCLASS_UScriptStruct))
+	void BindStruct(FBindId Id, const UStruct* Struct)
 	{
-		GUE.Defaults.Bind(Id, static_cast<const UScriptStruct*>(Struct));
-	}
-}
+		if (GUE.Customs.FindStruct(Id))
+		{
+			return;
+		}
 
-////////////////////////////////////////////////////////////////////////////////////////////////
+		if (GUE.Schemas.FindStruct(Id))
+		{
+			return;
+		}
 
-static void BindInitialTypes()
-{
-	UE_LOGFMT(LogPlainPropsUObject, Display, "Binding types to PlainProps schemas...");
+		FPropertyBinder Binder(Id, GetOccupancy(Struct));
+		if (const UStruct* Super = GetSuperToBind(Struct))
+		{
+			if (Binder.IsDense())
+			{
+				// Flatten inheritance chain for dense structs
+				for (; Super; Super = GetSuperToBind(Super))
+				{
+					BindMembers(/* out */ Binder, Super);	
+				}
+			}
+			else
+			{
+				FDeclId SuperId = GUE.Names.IndexDeclId(IndexType(Super));
+				Binder.BindSuper(SuperId);
+			}
+		}
 
-	// Declare all UEnums
-	for (TObjectIterator<UEnum> It; It; ++It)
-	{
-		DeclareEnum(*It);
+		BindMembers(/* out */ Binder, Struct);
+
+		GUE.Schemas.BindStruct(Id, Binder.GetMembers(), Binder.Declare());
+		
+		// Don't bind CDO defaults, object defaults are passed in from top and objects aren't owned by containers 
+		if (Struct->HasAnyCastFlags(CASTCLASS_UScriptStruct))
+		{
+			GUE.Defaults.Bind(Id, static_cast<const UScriptStruct*>(Struct));
+		}
 	}
 
-	// Bind all UScriptStruct/UClass/UFunction
-	static const FBindId SkipStructs[] = { GUE.Structs.VerseFunction };
-	for (TObjectIterator<UStruct> It; It; ++It)
+	void BindAllTypes()
 	{
-		const UStruct* Struct = *It;
-		FBindId Id = GUE.Names.IndexBindId(IndexType(Struct));
-		if (!Algo::Find(SkipStructs, Id))
+		// Declare all UEnums
+		for (TObjectIterator<UEnum> It; It; ++It)
 		{
-			BindStruct(Id, Struct);
+			DeclareEnum(*It);
+		}
+
+		// Bind all UScriptStruct/UClass/UFunction
+		static const FBindId SkipStructs[] = { GUE.Structs.VerseFunction };
+		for (TObjectIterator<UStruct> It; It; ++It)
+		{
+			const UStruct* Struct = *It;
+			FBindId Id = GUE.Names.IndexBindId(IndexType(Struct));
+			if (!Algo::Find(SkipStructs, Id))
+			{
+				BindStruct(Id, Struct);
+			}
 		}
 	}
-}
+};
+
+////////////////////////////////////////////////////////////////////////////////////////////////
 
 static void InitBatchedProperties()
 {
+	GUE.Defaults.BindZeroes(GUE.Structs.FieldPath, sizeof(FFieldPath), alignof(FFieldPath));
 	GUE.Defaults.BindZeroes(GUE.Structs.Name, sizeof(FName), alignof(FName));
 	GUE.Defaults.BindStatic(GUE.Structs.Text, &FText::GetEmpty());
 	GUE.Defaults.BindZeroes(GUE.Structs.ClassPtr, sizeof(TSubclassOf<UClass>), alignof(TSubclassOf<UClass>));
@@ -3639,13 +3747,70 @@ public:
 	, DeltaCtx{{GUE.Schemas, Customs}, Scratch, &GUE.Defaults}
 	{
 		SavedObjects.Reserve(NumReserve);
+		ObjectMap.Reserve(NumReserve);
 	}
 
-	FORCENOINLINE void Save(FBindId Id, const UObject* Object, const UObject* Arch)
+	// If Arch, delta serialize the matching part of the object inheritance chain
+	FBuiltStruct* SaveImpl(FBindId Id, const UStruct* Class, const UObject* Object, const UStruct* ArchClass, const UObject* Arch)
 	{
-		FBuiltStruct* Built = Arch ? SaveStructDelta(Object, Arch, Id, DeltaCtx) : SaveStruct(Object, Id, FlatCtx);
+		if (ArchClass == nullptr)
+		{
+			return SaveStruct(Object, Id, FlatCtx);
+		}
+		if (Class == ArchClass)
+		{
+			return SaveStructDelta(Object, Arch, Id, DeltaCtx);
+		}
+
+		Class = SkipEmptyBases(Class->GetInheritanceSuper());
+		FBindId SuperId = GUE.Names.IndexBindId(IndexType(Class));
+		FBuiltStruct* Super = SaveImpl(SuperId, Class, Object, ArchClass, Arch);
+		return SaveStructWithSuper(Object, Super, Id, FlatCtx);
+	}
+
+	void Save(FBindId Id, const UObject* Object, const UObject* Arch)
+	{
+		FBuiltStruct* Built;
+		if (Arch)
+		{
+			const UStruct* Class = SkipEmptyBases(Object->GetClass());
+			const UStruct* ArchClass = SkipEmptyBases(Arch->GetClass());
+			Built = SaveImpl(Id, Class, Object, ArchClass, Arch);
+		}
+		else
+		{
+			Built = SaveImpl(Id, Object->GetClass(), Object, nullptr, nullptr);
+		}
+		ObjectMap.Add(SavedObjects.Num());
 		SavedObjects.Emplace(Id, Built, Object);
 	}
+	
+	void Skip(FBindId Id, const UObject* Object)
+	{
+		ObjectMap.Add(~0u);
+	}
+
+	// Use existing SaveNameBatch() API for now, can be optimized later
+	void WriteStableNames(TArray64<uint8>& Out, TConstArrayView<FNameId> UsedNames) const
+	{
+		TArray<FDisplayNameEntryId> Entries;
+		TArray<int32> Numbers;
+		Entries.Reserve(UsedNames.Num());
+		Numbers.Reserve(UsedNames.Num());
+		for (FNameId NameId : UsedNames)
+		{
+			FSensitiveName SensitiveName = GUE.Names.ResolveName(NameId);
+			FName Name = SensitiveName.ToName();
+			Entries.Add(FDisplayNameEntryId(Name));
+			Numbers.Add(Name.GetNumber());
+		}
+
+		FBufferWriter NameMapArchive(nullptr, 0, EBufferWriterFlags::AllowResize | EBufferWriterFlags::TakeOwnership);
+		SaveNameBatch(Entries, NameMapArchive);
+
+		WriteNumAndArray(Out, TArrayView<uint8>(static_cast<uint8*>(NameMapArchive.GetWriterData()), static_cast<int32>(NameMapArchive.Tell())));
+		WriteNumAndArray(Out, Numbers);
+	}
 
 	TArray64<uint8>	Write(TArray<FStructId>* OutRuntimeIds = nullptr) const
 	{
@@ -3671,18 +3836,11 @@ public:
 		FWriter Writer(GUE.Names, Types, Schemas, Format);
 		TArray64<uint8> Out;
 
+		WriteInt(Out, Magics[0]);
 		// Write out FNames when using StableNames
 		if (!OutRuntimeIds)
 		{
-			TArray<FSensitiveName> UsedNames;
-			UsedNames.Reserve(Writer.GetUsedNames().Num());
-			for (FNameId Name : Writer.GetUsedNames())
-			{
-				UsedNames.Add(GUE.Names.ResolveName(Name));
-			}
-
-			WriteInt(Out, Magics[0]);
-			WriteNumAndArray(Out, TArrayView<const FSensitiveName, int32>(UsedNames));
+			WriteStableNames(Out, Writer.GetUsedNames());
 		}
 
 		// Write schemas
@@ -3707,6 +3865,9 @@ public:
 		// Write object terminator
 		WriteSkippableSlice(Out, TConstArrayView64<uint8>());
 		WriteInt(Out, Magics[4]);
+
+		// Write object map that allows for empty objects/(indices) without built data
+		WriteNumAndArray(Out, ObjectMap);
 		
 		return Out;
 	}
@@ -3720,6 +3881,7 @@ private:
 	};
 
 	TArray<FSavedObject>		SavedObjects;
+	TArray<uint32>				ObjectMap;
 	mutable FScratchAllocator	Scratch;
 	FSaveContext				FlatCtx;
 	FSaveContext				DeltaCtx;
@@ -3739,11 +3901,12 @@ public:
 	{
 		// Read and mount schemas
 		FByteReader It(Data);
-		check(It.Grab<uint32>() == Magics[1]);
+		verify(It.Grab<uint32>() == Magics[0]);
+		verify(It.Grab<uint32>() == Magics[1]);
 		It.SkipAlignmentPadding<uint32>();
 		uint32 SchemasSize = It.Grab<uint32>();
 		const FSchemaBatch* SavedSchemas = ValidateSchemas(It.GrabSlice(SchemasSize));
-		check(It.Grab<uint32>() == Magics[2]);
+		verify(It.Grab<uint32>() == Magics[2]);
 		
 		FSchemaBatchId Batch = MountReadSchemas(SavedSchemas);
 
@@ -3751,12 +3914,15 @@ public:
  		while (uint64 NumBytes = It.GrabVarIntU())
 		{	
 			FByteReader ObjIt(It.GrabSlice(NumBytes));
-			check(ObjIt.Grab<uint32>() == Magics[3]);
+			verify(ObjIt.Grab<uint32>() == Magics[3]);
 			FStructSchemaId Id = { ObjIt.Grab<uint32>() };
+			check(Id != FStructSchemaId{});
 			Objects.Add({ { Id, Batch }, ObjIt });
 		}
 		
-		check(It.Grab<uint32>() == Magics[4]);
+		verify(It.Grab<uint32>() == Magics[4]);
+		verify(GrabNumAndArray<uint32>(It).Num() == Objects.Num());
+		It.CheckEmpty();
 		check(!Objects.IsEmpty());
 
 		// Finally create load plans
@@ -3796,6 +3962,101 @@ private:
 	}
 };
 
+class FPersistentBatchLoader
+{
+public:
+	FPersistentBatchLoader(const FCustomBindings& Customs, FMemoryView Data)
+	{
+		// Read ids
+		FByteReader It(Data);
+		verify(It.Grab<uint32>() == Magics[0]);
+		ReadStableNames(It);
+
+		// Read and mount schemas
+		verify(It.Grab<uint32>() == Magics[1]);
+		It.SkipAlignmentPadding<uint32>();
+		uint32 SchemasSize = It.Grab<uint32>();
+		const FSchemaBatch* SavedSchemas = ValidateSchemas(It.GrabSlice(SchemasSize));
+		verify(It.Grab<uint32>() == Magics[2]);
+		
+		// Bind saved ids to runtime ids, make new schemas with new ids and mount them
+		FIdTranslator RuntimeIds(GUE.Names, MakeConstArrayView(Names), *SavedSchemas);
+		TranslatedSchemas = CreateTranslatedSchemas(*SavedSchemas, RuntimeIds.Translation);
+		FSchemaBatchId Batch = MountReadSchemas(TranslatedSchemas);
+
+		// Read objects
+ 		while (uint64 NumBytes = It.GrabVarIntU())
+		{	
+			FByteReader ObjIt(It.GrabSlice(NumBytes));
+			verify(ObjIt.Grab<uint32>() == Magics[3]);
+			FStructSchemaId Id = { ObjIt.Grab<uint32>() };
+			Objects.Add({ { Id, Batch }, ObjIt });
+		}
+		
+		verify(It.Grab<uint32>() == Magics[4]);
+		ObjectMap = GrabNumAndArray<uint32>(It);
+		check(ObjectMap.Num() >= Objects.Num());
+		check(!Objects.IsEmpty());
+
+		// Finally create load plans
+		TConstArrayView<FStructId> LoadStructIds = RuntimeIds.Translation.GetStructIds(NumStructSchemas(Batch));
+		Plans = CreateLoadPlans(Batch, Customs, GUE.Schemas, LoadStructIds, ESchemaFormat::StableNames);
+	}
+
+	void ReadStableNames(FByteReader& It)
+	{
+		TConstArrayView64<uint8> NameBatch = GrabNumAndArray<uint8>(It);
+		TConstArrayView<int32> NameNumbers = GrabNumAndArray<int32>(It);
+		FMemoryReaderView Ar(MakeArrayView(NameBatch.GetData(), NameBatch.Num()));
+		TArray<FDisplayNameEntryId> NameEntries = LoadNameBatch(Ar);
+		Names.Reserve(NameEntries.Num());
+		check(NameEntries.Num() == NameNumbers.Num());
+		for (int64 Idx = 0, Num = NameEntries.Num(); Idx < Num; ++Idx)
+		{
+			Names.Emplace(NameEntries[Idx].ToName(NameNumbers[Idx]));
+		}
+	}
+
+	~FPersistentBatchLoader()
+	{
+		Plans.Reset();
+		UnmountReadSchemas(Objects[0].Schema.Batch);
+	}
+
+	FORCENOINLINE void Load(UObject* Dst)
+	{
+		FStructView In = Objects[LoadIdx];
+		LoadStruct(Dst, In.Values, In.Schema.Id, *Plans);
+		++LoadIdx;
+	}
+
+	bool HasObject(int32 ObjectIndex)
+	{
+		return ObjectMap[ObjectIndex] != ~0u;
+	}
+
+	FORCENOINLINE void LoadObject(UObject* Dst, int32 ObjectIndex)
+	{
+		FStructView In = Objects[ObjectMap[ObjectIndex]];
+		LoadStruct(Dst, In.Values, In.Schema.Id, *Plans);
+	}
+
+private:
+	const FSchemaBatch*			TranslatedSchemas = nullptr;
+	TArray<FSensitiveName>		Names;
+	FLoadBatchPtr				Plans;
+	TArray<FStructView>			Objects;
+	TConstArrayView<uint32>		ObjectMap;
+	int32						LoadIdx = 0;
+
+	template<typename T>
+	static TConstArrayView<T> GrabNumAndArray(/* in-out */ FByteReader& It)
+	{
+		uint32 Num = It.Grab<uint32>();
+		return MakeArrayView(reinterpret_cast<const T*>(It.GrabBytes(Num * sizeof(T))), Num);
+	}
+};
+
 ////////////////////////////////////////////////////////////////////////////////////////////////
 
 // Similar to PlainProps FMemoryBatch but for storing FArchive property serialization results
@@ -3947,7 +4208,7 @@ static bool IncludeClass(UClass* Class)
 	};
 	static const FName SuperExclusions[] = { "LevelScriptActor" };
 
-	if (!ShouldBind(Class) || !!Algo::Find(Exclusions, Class->GetFName()))
+	if (!!Algo::Find(Exclusions, Class->GetFName()))
 	{
 		return false;
 	}
@@ -4048,10 +4309,9 @@ static void RoundtripText(
 	const FBatchIds& BatchIds,
 	TConstArrayView<FStructView> Objects,
 	TConstArrayView<FInstance> Instances,
-	ESchemaFormat Format)
+	ESchemaFormat Format,
+	FStringView AssetName = {})
 {
-	check(Objects.Num() == Instances.Num());
-
 	// Print yaml
 	UE_LOGFMT(LogPlainPropsUObject, Display, "Printing to PlainProps text using {Format}...", ToString(Format));
 	TUtf8StringBuilder<256> Yaml;
@@ -4060,11 +4320,14 @@ static void RoundtripText(
 	FUtf8StringView YamlView = Yaml.ToView();
 
 	// Write to file
-	const FString Filename = FPaths::ProjectSavedDir() / TEXT("PlainProps") /
-		(Format == ESchemaFormat::InMemoryNames ? TEXT("InMemoryNames.yaml") : TEXT("StableNames.yaml"));
-	if (TUniquePtr<FArchive> FileWriter(IFileManager::Get().CreateFileWriter(*Filename)); FileWriter)
+	if (AssetName.Len() == 0)
+	{
+		AssetName = (Format == ESchemaFormat::InMemoryNames ? TEXTVIEW("DiffTest_InMemoryNames") : TEXTVIEW("DiffTest_StableNames"));
+	}
+	const FString YamlPath = FPaths::Combine(FPaths::ProjectSavedDir(), TEXT("PlainProps"), AssetName) + TEXT(".yaml");
+	if (TUniquePtr<FArchive> FileWriter(IFileManager::Get().CreateFileWriter(*YamlPath)); FileWriter)
 	{
-		UE_LOGFMT(LogPlainPropsUObject, Display, "Writing {KB}KB yaml as {Filename}...", Yaml.Len() >> 10, *Filename);
+		UE_LOGFMT(LogPlainPropsUObject, Display, "Writing {KB}KB yaml as {YamlPath}...", Yaml.Len() >> 10, *YamlPath);
 		FileWriter->Serialize((void*)YamlView.GetData(), YamlView.Len());
 	}
 
@@ -4074,7 +4337,7 @@ static void RoundtripText(
 	TArray<FStructView> ParsedObjects;
 	FSchemaBatchId ParsedBatch = ParseBatch(Data, ParsedObjects, YamlView);
 
-	if (Format == ESchemaFormat::StableNames)
+	if (Format == ESchemaFormat::StableNames && Instances.Num())
 	{
 		UE_LOGFMT(LogPlainPropsUObject, Display, "Diffing PlainProps parsed objects using {Format}...", ToString(Format));
 
@@ -4083,6 +4346,7 @@ static void RoundtripText(
 
 		// Diff objects
 		check(Objects.Num() == ParsedObjects.Num());
+		check(Objects.Num() == Instances.Num());
 		const int32 NumObjects = FMath::Min(Objects.Num(), ParsedObjects.Num());
 		uint32 NumDiffs = 0;
 		TUtf8StringBuilder<256> Diffs;
@@ -4115,34 +4379,33 @@ public:
 	FBatchTextRoundtripper(FMemoryView Data, ESchemaFormat InFormat) : Format(InFormat)
 	{
 		FByteReader It(Data);
+		verify(It.Grab<uint32>() == Magics[0]);
 
 		// Read FNames when using Stable Names
-		TConstArrayView<FSensitiveName> Names;
 		if (Format == ESchemaFormat::StableNames)
 		{
-			verify(It.Grab<uint32>() == Magics[0]);
-			Names = GrabNumAndArray<FSensitiveName>(It);
+			ReadStableNames(It);
 		}
 		
 		// Read and mount schemas
-		check(It.Grab<uint32>() == Magics[1]);
+		verify(It.Grab<uint32>() == Magics[1]);
 		It.SkipAlignmentPadding<uint32>();
 		const uint32 SchemasSize = It.Grab<uint32>();
 		FMemoryView SavedSchemasView = It.GrabSlice(SchemasSize);
 		const FSchemaBatch* SavedSchemas = ValidateSchemas(SavedSchemasView);
-		check(It.Grab<uint32>() == Magics[2]);
+		verify(It.Grab<uint32>() == Magics[2]);
 		FSchemaBatchId Batch = MountReadSchemas(SavedSchemas);
 		
 		// Read objects
 		while (uint64 NumBytes = It.GrabVarIntU())
 		{	
 			FByteReader ObjIt(It.GrabSlice(NumBytes));
-			check(ObjIt.Grab<uint32>() == Magics[3]);
+			verify(ObjIt.Grab<uint32>() == Magics[3]);
 			FStructSchemaId Id = { ObjIt.Grab<uint32>() };
 			Objects.Add({ { Id, Batch }, ObjIt });
 		}
 		
-		check(It.Grab<uint32>() == Magics[4]);
+		verify(It.Grab<uint32>() == Magics[4]);
 		check(!Objects.IsEmpty());
 
 		// Create BatchIds 
@@ -4156,18 +4419,38 @@ public:
 		}
 	}
 
+	void ReadStableNames(FByteReader& It)
+	{
+		TConstArrayView<uint8> NameBatch = GrabNumAndArray<uint8>(It);
+		TConstArrayView<int32> Numbers = GrabNumAndArray<int32>(It);
+		FMemoryReaderView Ar(NameBatch);
+		TArray<FDisplayNameEntryId> Entries = LoadNameBatch(Ar);
+		Names.Reserve(Entries.Num());
+		check(Entries.Num() == Numbers.Num());
+		for (int64 Idx = 0, Num = Entries.Num(); Idx < Num; ++Idx)
+		{
+			Names.Emplace(Entries[Idx].ToName(Numbers[Idx]));
+		}
+	}
+
 	~FBatchTextRoundtripper()
 	{
 		UnmountReadSchemas(BatchIds->GetBatchId());
 	}
 
-	void RoundtripText(TConstArrayView<FInstance> Instances) const
+	void RoundtripTextForDiffTest(TConstArrayView<FInstance> Instances) const
+	{
+		PlainProps::UE::RoundtripText(*BatchIds, Objects, Instances, Format, {});
+	}
+
+	void RoundtripTextForPackage(FStringView PackageName) const
 	{
-		PlainProps::UE::RoundtripText(*BatchIds, Objects, Instances, Format);
+		PlainProps::UE::RoundtripText(*BatchIds, Objects, {}, Format, PackageName);
 	}
 
 private:
 	TArray<FStructView>			Objects;
+	TArray<FSensitiveName>		Names;
 	TUniquePtr<FBatchIds>		BatchIds;
 	ESchemaFormat				Format;
 	
@@ -4183,24 +4466,411 @@ private:
 
 struct FDiffDebug
 {
-	volatile FInstance* Last;
+	TArray<FInstance*> Instances;
 	volatile const UTF8CHAR* Str;
 };
 static FDiffDebug GPpDiff, GTpsDiff, GUpsDiff;
 
-static int32 Roundtrip(ERoundtrip Options)
+static void DumpUObjectToJsonFile(const UObject* Object)
+{
+	FString JsonString;
+	if (FJsonObjectConverter::UStructToJsonObjectString(Object->GetClass(), Object, JsonString))
+	{
+		TStringBuilder<256> JsonPath;
+		JsonPath << FPaths::Combine(FPaths::ProjectSavedDir(), TEXT("PlainProps"), Object->GetPackage()->GetName());
+		JsonPath << "." << Object->GetName() + ".json";
+
+		UE_LOGFMT(LogPlainPropsUObject, Display, "Writing {JsonPath}...", JsonPath);
+		FFileHelper::SaveStringToFile(JsonString, JsonPath.ToString(), FFileHelper::EEncodingOptions::ForceUTF8WithoutBOM);
+	}
+}
+
+////////////////////////////////////////////////////////////////////////////////
+
+#if WITH_EDITOR
+
+struct FLinkerPropertyBatch
+{
+	explicit FLinkerPropertyBatch(const FLinkerSave& InLinkerSave)
+		: LinkerSave(&InLinkerSave)
+	{
+	}
+
+	explicit FLinkerPropertyBatch(FLinkerLoad& InLinkerLoad)
+		: LinkerLoad(&InLinkerLoad)
+	{
+	}
+
+	static FSoleMemberSpec Spec(FName*, FDeclId Id)
+	{
+		return {Id, GUE.Members.Id, Specify<uint64>()};
+	}
+
+	void Save(FMemberBuilder& Out, FName In, const FSaveContext&)
+	{
+		// see FLinkerSave::MapName
+		int32 Index = LinkerSave->NameIndices.FindRef(In.GetDisplayIndex(), INDEX_NONE);
+		int32 Number = In.GetNumber();
+		uint64 Id;
+		if (Index == INDEX_NONE)
+		{
+			UE_LOGFMT(LogPlainPropsUObject, Warning, "FName {Name} is not mapped.", In);
+			Id = 6666666;
+		}
+		else
+		{
+			Id = Index != INDEX_NONE ? ((uint64(Number) << 32) | Index) : 6666666;
+		}
+		Out.Add(GUE.Members.Id, Id);
+	}
+
+	void Load(FName& Out, FStructLoadView In)
+	{
+		uint64 Id;
+		LoadSole<uint64>(&Id, In);
+		if (Id != 6666666)
+		{
+			int32 Index = Id & 0xFFFFFFFF;
+			int32 Number = Id >> 32;
+			FNameEntryId Entry = LinkerLoad->NameMap[Index];
+			FName Name = FName::CreateFromDisplayId(Entry, Number);
+			Out = Name;
+		}
+		else
+		{
+			Out = FName();
+		}
+	}
+
+	static FSoleMemberSpec Spec(FText*, FDeclId Id)
+	{
+		return {Id, GUE.Members.Id, Specify<int32>()};
+	}
+
+	void Save(FMemberBuilder& Out, const FText& In, const FSaveContext&)
+	{
+		Out.Add(GUE.Members.Id, Texts.Num());
+		Texts.Add(In);
+	}
+
+	void Load(FText& Out, FStructLoadView In) const
+	{
+		int32 Dummy = LoadSole<int32>(In);
+		Out = {};//TODO: Texts[LoadSole<int32>(In)];
+	}
+
+	static FSoleMemberSpec Spec(FObjectHandle*, FDeclId Id)
+	{
+		return {Id, GUE.Members.Id, Specify<int32>()};
+	}
+
+	void Save(FMemberBuilder& Out, FObjectHandle In, const FSaveContext&)
+	{
+		// see FLinkerSave::MapObject
+		UObject* Object = FObjectPtr(In).Get();
+		FPackageIndex Index = LinkerSave->ObjectIndicesMap.FindRef(Object);
+		int32 Idx = *reinterpret_cast<int32*>(&Index);
+		if (Object && Index.IsNull())
+		{
+			UE_LOGFMT(LogPlainPropsUObject, Warning, "UObject {Object} is not mapped.", Object->GetFullName());
+			Idx = 9999999;
+		}
+		Out.Add(GUE.Members.Id, Idx);
+	}
+
+	void Load(FObjectHandle& Out, FStructLoadView In)
+	{
+		int32 Idx;
+		LoadSole<int32>(&Idx, In);
+		if (Idx != 9999999)
+		{
+			FPackageIndex PackageIdx = *reinterpret_cast<FPackageIndex*>(&Idx);
+			if (PackageIdx.IsImport())
+			{
+				UObject* Object = LinkerLoad->ImportMap[PackageIdx.ToImport()].XObject;
+				if (Object == nullptr)
+				{
+#if PP_WITH_ENGINE_HOOKS
+					// Object = LinkerLoad->ResolveResource(PackageIdx);
+#endif
+				}
+				Out = ::UE::CoreUObject::Private::MakeObjectHandle(Object);
+			}
+			else if (PackageIdx.IsExport())
+			{
+				UObject* Object = LinkerLoad->ExportMap[PackageIdx.ToExport()].Object;
+				if (Object == nullptr)
+				{
+#if PP_WITH_ENGINE_HOOKS
+					// Object = LinkerLoad->ResolveResource(PackageIdx);
+#endif
+				}
+				Out = ::UE::CoreUObject::Private::MakeObjectHandle(Object);
+			}
+			else
+			{
+				check(IsObjectHandleNull(Out));
+			}
+		}
+		else
+		{
+			Out = FObjectHandle();
+		}
+	}
+
+	static FSoleMemberSpec Spec(FWeakObjectPtr*, FDeclId Id)
+	{
+		return {Id, GUE.Members.Id, Specify<uint64>()};
+	}
+
+	void Save(FMemberBuilder& Out, const FWeakObjectPtr& In, const FSaveContext& Ctx)
+	{
+		// Save ObjectSerialNumber + ObjectIndex a single uint64
+		UObject* Object = In.Get(true);
+		Save(Out, ::UE::CoreUObject::Private::MakeObjectHandle(Object), Ctx);
+	}
+
+	void Load(FWeakObjectPtr& Out, FStructLoadView In)
+	{
+		FObjectHandle Handle;
+		Load(Handle, In);
+		Out = ::UE::CoreUObject::Private::ResolveObjectHandle(Handle);
+	}
+
+	static FSoleMemberSpec Spec(FSoftObjectPtr*, FDeclId Id)
+	{
+		return {Id, GUE.Members.Id, FMemberSpec(GUE.Structs.SoftObjectPath)};
+	}
+
+	static void Save(FMemberBuilder& Out, const FSoftObjectPtr& In, const FSaveContext& Ctx)
+	{
+		FBuiltStruct* SoftPath = SaveStruct(&In.GetUniqueID(), GUE.Structs.SoftObjectPath, Ctx);
+		Out.AddStruct(GUE.Members.Id, GUE.Structs.SoftObjectPath, SoftPath);
+	}
+
+	static void Load(FSoftObjectPtr& Out, FStructLoadView In)
+	{
+		Out.ResetWeakPtr();
+		LoadSoleStruct(&Out.GetUniqueID(), In);
+	}
+
+	static FSoleMemberSpec Spec(FLazyObjectPtr*, FDeclId Id)
+	{
+		return {Id, GUE.Members.Id, FMemberSpec(GUE.Structs.Guid)};
+	}
+	
+	static void Save(FMemberBuilder& Out, const FLazyObjectPtr& In, const FSaveContext& Ctx)
+	{
+		FBuiltStruct* Guid = SaveStruct(&In.GetUniqueID(), GUE.Structs.Guid, Ctx);
+		Out.AddStruct(GUE.Members.Id, GUE.Structs.Guid, Guid);
+	}
+
+	static void Load(FLazyObjectPtr& Out, FStructLoadView In)
+	{
+		Out.ResetWeakPtr();
+		LoadSoleStruct(&Out.GetUniqueID(), In);
+	}
+
+private:
+	union
+	{
+		const FLinkerSave* LinkerSave;
+		FLinkerLoad* LinkerLoad;
+		const FLinker* Linker;
+	};
+
+	TArray<FText> Texts; // Tricky to serialize non-intrusively
+};
+
+struct FPlainPropsLinkerSandbox
 {
-	UE_LOGFMT(LogPlainPropsUObject, Display, "Gathering all non-empty UObjects...");
-	static constexpr EObjectFlags SkipFlags = RF_ClassDefaultObject | RF_MirroredGarbage | RF_InheritableComponentTemplate;
+	static constexpr uint32 LoadFlags = (LOAD_Async | LOAD_NoVerify | LOAD_SkipLoadImportedPackages);
+	static constexpr uint32 CloneNumber = 8848;
+
+	struct FSavedPackage
+	{
+		FName PackageName;
+		FPackagePath PackagePath;
+		FSharedBuffer PackageData;
+		TArray64<uint8> PlainPropsData;
+	};
+
+	struct FLinkerBatch
+	{
+		FLinkerBatch(const FSavedPackage& Saved)
+			: PackagePath(Saved.PackagePath)
+			, PackageName(Saved.PackageName.GetNumber() ? FName(*Saved.PackageName.ToString(), CloneNumber) : FName(Saved.PackageName, CloneNumber))
+			, Package(NewObject<UPackage>(/*Outer*/nullptr, PackageName))
+			, LinkerReader(static_cast<const uint8*>(Saved.PackageData.GetData()), Saved.PackageData.GetSize(), ELargeMemoryReaderFlags::Persistent, PackageName)
+			, LinkerLoad(FLinkerLoad::CreateLinker(FUObjectThreadContext::Get().GetSerializeContext(), Package, PackagePath, LoadFlags, &LinkerReader))
+			, Properties(*LinkerLoad)
+			, Customs(Properties, GUE.Customs)
+			, BatchView(MakeMemoryView(Saved.PlainPropsData))
+			, BatchLoader(Customs.Overlay, BatchView)
+		{
+		}
+
+		FPackagePath PackagePath;
+		FName PackageName;
+		UPackage* Package;
+		FLargeMemoryReader LinkerReader;
+		FLinkerLoad* LinkerLoad;
+		FLinkerPropertyBatch Properties;
+		TCustomPropertyBindings<FLinkerPropertyBatch> Customs;
+		FMemoryView BatchView;
+		FPersistentBatchLoader BatchLoader;
+	};
+
+	static TMap<FName, FSavedPackage> SavedPackages;
+	static TMap<FName, TUniquePtr<FLinkerBatch>> LoadBatches;
+
+	static void SavePlainProps(const FLinkerSave& LinkerSave)
+	{
+		FName PackageName = LinkerSave.LinkerRoot->GetFName();
+
+		UE_LOGFMT(LogPlainPropsUObject, Display, " Saving extension data for {PackageName}...", PackageName);
+		FLinkerPropertyBatch Properties(LinkerSave);
+		TCustomPropertyBindings<FLinkerPropertyBatch> Customs(Properties, GUE.Customs);
+		FBatchSaver BatchSaver(Customs.Overlay, LinkerSave.ExportMap.Num());
+		for (const FObjectExport& Export : LinkerSave.ExportMap)
+		{
+			const UStruct* Struct = SkipEmptyBases(Export.Object->GetClass());
+			FBindId Id = GUE.Names.IndexBindId(PlainProps::UE::IndexType(Struct));
+			if (GUE.Schemas.FindStruct(Id))
+			{
+				BatchSaver.Save(Id, Export.Object, Export.Object->GetArchetype());
+			}
+			else
+			{
+				BatchSaver.Skip(Id, Export.Object);
+			}
+		}
+		TArray64<uint8> StableData = BatchSaver.Write();
+
+		FSavedPackage& SavedPackage = SavedPackages.FindOrAdd(PackageName);
+		SavedPackage.PackageName = PackageName;
+		SavedPackage.PackagePath = LinkerSave.LinkerRoot->GetLoadedPath();
+		SavedPackage.PlainPropsData = MoveTemp(StableData);
+	}
+
+	static void SavePackageToMemory(const FLinkerSave& LinkerSave, uint8* Data, int64 Size)
+	{
+		FName PackageName = LinkerSave.LinkerRoot->GetFName();
+		FSavedPackage& SavedPackage = SavedPackages.FindChecked(PackageName);
+		SavedPackage.PackageData = FSharedBuffer::TakeOwnership(Data, Size, FMemory::Free);
+	}
+
+	static bool OpenPackageReader(const FPackagePath& PackagePath, FArchive*& Ar)
+	{
+		FName PackageName = PackagePath.GetPackageFName();
+		if (FSavedPackage* SavedData = SavedPackages.Find(PackageName))
+		{
+			Ar = new FLargeMemoryReader(static_cast<const uint8*>(SavedData->PackageData.GetData()), SavedData->PackageData.GetSize(), ELargeMemoryReaderFlags::Persistent, PackageName);
+			return true;
+		}
+		return false;
+	}
+
+	static void MountPackageFromMemory(FName PackageName)
+	{
+		if (FSavedPackage* SavedData = SavedPackages.Find(PackageName))
+		{
+			FLinkerBatch* LinkerBatch = new FLinkerBatch(*SavedData);
+			LoadBatches.Emplace(PackageName, LinkerBatch);
+		}
+	}
+
+	static void LoadPackageFromMemory(FName PackageName)
+	{
+		if (TUniquePtr<FLinkerBatch>* LinkerBatch = LoadBatches.Find(PackageName))
+		{
+			UE_LOGFMT(LogPlainPropsUObject, Display, "  Loading {Package}...", PackageName);
+			LoadPackage((*LinkerBatch)->Package, (*LinkerBatch)->PackagePath, LoadFlags);
+		}
+	}
+
+	static bool LoadObject(FStructuredArchive::FSlot& Slot, uint8* Data)
+	{
+		UObject* Object = (UObject*)Data;
+		FName PackageName = Object->GetPackage()->GetFName();
+		if (PackageName.GetNumber() == CloneNumber)
+		{
+			PackageName.SetNumber(0);
+		}
+
+		TUniquePtr<FLinkerBatch>* LinkerBatch = LoadBatches.Find(PackageName);
+		if (!LinkerBatch)
+		{
+			return false;
+		}
+
+		FPersistentBatchLoader& BatchLoader = (*LinkerBatch)->BatchLoader;
+		FLinkerLoad* LinkerLoad = (*LinkerBatch)->LinkerLoad;
+		check(LinkerLoad == Object->GetLinker());
+
+		FUObjectSerializeContext* SerializeContext = FUObjectThreadContext::Get().GetSerializeContext();
+		check(Object == SerializeContext->SerializedObject);
+
+		const int32 ExportIndex = Object->GetLinkerIndex();
+		const FObjectExport& Export = LinkerLoad->ExportMap[ExportIndex];
+		check(Export.Object == Object);
+
+		BatchLoader.LoadObject(Object, ExportIndex);
+
+		const int32 Pos = Slot.GetUnderlyingArchive().Tell();
+		check(Pos == SerializeContext->SerializedObjectScriptStartOffset);
+		const int32 ScriptSerializationSize = Export.ScriptSerializationEndOffset - Export.ScriptSerializationStartOffset;
+		check(ScriptSerializationSize >= 0);
+		Slot.GetUnderlyingArchive().Seek(Pos + ScriptSerializationSize);
+		return true;
+	}
+};
+
+TMap<FName, FPlainPropsLinkerSandbox::FSavedPackage> FPlainPropsLinkerSandbox::SavedPackages;
+TMap<FName, TUniquePtr<FPlainPropsLinkerSandbox::FLinkerBatch>> FPlainPropsLinkerSandbox::LoadBatches;
+
+#endif
+
+////////////////////////////////////////////////////////////////////////////////
+
+static int32 RoundtripObjects(TConstArrayView<UObject*> Objects, ERoundtrip Options)
+{
+	UE_LOGFMT(LogPlainPropsUObject, Display, "Binding types to PlainProps schemas...");
+	FTypeBinder StructBinder;
+	StructBinder.BindAllTypes();
+
 	TArray<FInstance> Instances;
-	for (TObjectIterator<UObject> It(SkipFlags); It; ++It)
+	if (Objects.Num())
 	{
-		UObject* Object = *It;
-		if (IncludeClass(Object->GetClass()))
+		UE_LOGFMT(LogPlainPropsUObject, Display, "Gathering UObjects in Assets...");
+		for (const UObject* Object : Objects)
 		{
-			Instances.Emplace(Object->GetPathName(), Object);
+			ForEachObjectWithPackage(Object->GetPackage(), [&Instances,&StructBinder](UObject* Obj)
+			{
+				UClass* Class = Obj->GetClass();
+				if (StructBinder.ShouldBind(Class) && IncludeClass(Class))
+				{
+					Instances.Emplace(Obj->GetPathName(), Obj);
+				}
+				return true;
+			});
 		}
 	}
+	else
+	{
+		UE_LOGFMT(LogPlainPropsUObject, Display, "Gathering all non-empty UObjects...");
+		static constexpr EObjectFlags SkipFlags = RF_ClassDefaultObject | RF_MirroredGarbage | RF_InheritableComponentTemplate;
+		for (TObjectIterator<UObject> It(SkipFlags); It; ++It)
+		{
+			UObject* Object = *It;
+			UClass* Class = Object->GetClass();
+			if (StructBinder.ShouldBind(Class) && IncludeClass(Class))
+			{
+				Instances.Emplace(Object->GetPathName(), Object);
+			}
+		}
+	}
+
 	UE_LOGFMT(LogPlainPropsUObject, Display, "Sorting {Num} UObjects ...", Instances.Num());
 	Algo::Sort(Instances, [](const FInstance& A, const FInstance& B)
 	{
@@ -4245,7 +4915,7 @@ static int32 Roundtrip(ERoundtrip Options)
 		if (EnumHasAnyFlags(Options, ERoundtrip::TextMemory))
 		{
 			FBatchTextRoundtripper MemoryBatch(MakeMemoryView(Plain.Data), ESchemaFormat::InMemoryNames);
-			MemoryBatch.RoundtripText(Instances);
+			MemoryBatch.RoundtripTextForDiffTest(Instances);
 		}
 		if (EnumHasAnyFlags(Options, ERoundtrip::TextStable))
 		{
@@ -4253,7 +4923,7 @@ static int32 Roundtrip(ERoundtrip Options)
 			TArray64<uint8> StableData = Batch.Write();
 
 			FBatchTextRoundtripper StableBatch(MakeMemoryView(StableData), ESchemaFormat::StableNames);
-			StableBatch.RoundtripText(Instances);
+			StableBatch.RoundtripTextForDiffTest(Instances);
 		}
 	}
 
@@ -4268,23 +4938,29 @@ static int32 Roundtrip(ERoundtrip Options)
 		// Diff original vs PlainProps
 		UE_LOGFMT(LogPlainPropsUObject, Display, "Diffing UObjects roundtripped via PlainProps...");
 		FDiffContext DiffCtx = {{ GUE.Schemas, Customs.Overlay }};
-		TUtf8StringBuilder<256> PpDiffs;
+		TUtf8StringBuilder<256> PpDiffPaths;
 		for (FInstance& Instance : Instances)
 		{
 			if (DiffStructs(Instance.Orig, Instance.PP, Instance.Id, /* in-out */ DiffCtx))
 			{
-				GPpDiff.Last = &Instance;
-				PrintDiff(/* out */ PpDiffs, GUE.Names, DiffCtx.Out);
+				PpDiffPaths.Append("\n");
+				GPpDiff.Instances.Add(&Instance);
+				PrintDiff(/* out */ PpDiffPaths, GUE.Names, DiffCtx.Out);
+
 				//Batch.Reload(Instance.UPS, &Instance - Instances.GetData());
 				DiffCtx.Out.Reset();
-				PpDiffs.Append(" in ");
-				PpDiffs.Append(Instance.Orig->GetFullName());
-				PpDiffs.Append("\n");
+				PpDiffPaths.Append(" in ");
+				PpDiffPaths.Append(Instance.Orig->GetFullName());
 				++NumPpDiffs;
 			}
 		}
-		GPpDiff.Str = PpDiffs.GetData();
-		UE_LOGFMT(LogPlainPropsUObject, Display, "Detected {Diffs} diffs in {Objs} UObjects saved in a {KB}KB value stream using PlainProps\n{DiffText}", NumPpDiffs, Instances.Num(), Plain.Data.NumBytes() / 1024, PpDiffs.ToString());
+		GPpDiff.Str = PpDiffPaths.GetData();
+		UE_LOGFMT(LogPlainPropsUObject, Display, "Detected {Diffs} diffs in {Objs} UObjects saved in a {KB}KB value stream using PlainProps{DiffText}", NumPpDiffs, Instances.Num(), Plain.Data.NumBytes() / 1024, PpDiffPaths.ToString());
+		for (const FInstance* Instance : GPpDiff.Instances)
+		{
+			DumpUObjectToJsonFile(Instance->Orig);
+			DumpUObjectToJsonFile(Instance->PP);
+		}
 	}
 
 	FArchivedProperties Tps;
@@ -4323,7 +4999,7 @@ static int32 Roundtrip(ERoundtrip Options)
 
 			if (DiffStructs(Instance.Orig, Instance.TPS, Instance.Id, /* in-out */ DiffCtx))
 			{
-				GTpsDiff.Last = &Instance;
+				GTpsDiff.Instances.Add(&Instance);
 				PrintDiff(/* out */ TpsDiffs, GUE.Names, DiffCtx.Out);
 				
 				DiffCtx.Out.Reset();
@@ -4378,7 +5054,7 @@ static int32 Roundtrip(ERoundtrip Options)
 
 			if (DiffStructs(Instance.Orig, Instance.UPS, Instance.Id, /* in-out */ DiffCtx))
 			{
-				GUpsDiff.Last = &Instance;
+				GUpsDiff.Instances.Add(&Instance);
 				PrintDiff(/* out */ UpsDiffs, GUE.Names, DiffCtx.Out);
 				
 				DiffCtx.Out.Reset();
@@ -4395,29 +5071,155 @@ static int32 Roundtrip(ERoundtrip Options)
 	return NumPpDiffs;
 }
 
-static int32 TestBindings(ERoundtrip Options)
+////////////////////////////////////////////////////////////////////////////////
+
+#if WITH_EDITOR
+static int32 RoundtripPackages(TConstArrayView<UObject*> Objects, ERoundtrip Options)
 {
+	TArray<UPackage*> Packages;
+	TArray<UObject*> AllPackages;
+
+	// Collect all packages
+	if (Objects.Num())
+	{
+		for (UObject* Object : Objects)
+		{
+			AllPackages.Add(Object->GetPackage());
+		}
+	}
+	else
+	{
+		GetObjectsOfClass(UPackage::StaticClass(), AllPackages);
+	}
+
+	// Filter out in-memory and BP packages
+	for (UObject* PackageObj : AllPackages)
+	{
+		if (!PackageObj->HasAnyFlags(RF_WasLoaded))
+		{
+			continue;
+		}
+		UPackage* Package = Cast<UPackage>(PackageObj);
+		if (Package->HasAnyPackageFlags(PKG_CompiledIn))
+		{
+			continue;
+		}
+		bool bContainsBP = false;
+		ForEachObjectWithPackage(Package, [&bContainsBP](UObject* Object)
+		{
+			if (Object->GetClass()->HasAnyClassFlags(CLASS_CompiledFromBlueprint))
+			{
+				bContainsBP = true;
+				return false;
+			}
+			return true;
+		});
+		if (bContainsBP)
+		{
+			continue;
+		}
+		Packages.Add(Package);
+	}
+
+	UE_LOGFMT(LogPlainPropsUObject, Display, "Binding types to PlainProps schemas...");
+	FTypeBinder{ .ExcludedPropertyFlags = CPF_Transient | CPF_Deprecated }.BindAllTypes();
+
+#if PP_WITH_ENGINE_HOOKS
+	FSavePackageSettings& SavePackageSettings = FSavePackageSettings::GetDefaultSettings();
+	SavePackageSettings.ExternalWriteExportsFunc = &FPlainPropsLinkerSandbox::SavePlainProps;
+	SavePackageSettings.ExternalWriteFileFunc = &FPlainPropsLinkerSandbox::SavePackageToMemory;
+	GSerializePropertiesOverride = &FPlainPropsLinkerSandbox::LoadObject;
+	GLoadFromPackageFileOverride = &FPlainPropsLinkerSandbox::OpenPackageReader;
+#endif
+
+	UE_LOGFMT(LogPlainPropsUObject, Display, "Saving {SaveNum} packages...", Packages.Num());
+	for (UPackage* Package : Packages)
+	{
+		UObject* Asset = Package->FindAssetInPackage();
+		FSavePackageArgs SaveArgs;
+		SaveArgs.TopLevelFlags = RF_Standalone;
+		SaveArgs.SaveFlags = SAVE_Async; // Write progress to temp memory instead of a temp file
+		const FString& Extension = Package->ContainsMap() ? FPackageName::GetMapPackageExtension() : FPackageName::GetAssetPackageExtension();
+		FString TempFilename = FPaths::Combine(*FPaths::ProjectSavedDir(), TEXT("PlainProps"), Package->GetName() + Extension);
+		// FString ActualFilename = FPackageName::LongPackageNameToFilename(Asset->GetPackage()->GetName(), Extension);
+
+		GEditor->SavePackage(Package, Asset, *TempFilename, SaveArgs);
+	}
+
+	FPlainPropsLinkerSandbox::SavedPackages.Reserve(Packages.Num());
+	FPlainPropsLinkerSandbox::LoadBatches.Reserve(Packages.Num());
+
+	UE_LOGFMT(LogPlainPropsUObject, Display, "Mounting {SaveNum} packages...", Packages.Num());
+	for (UPackage* Package : Packages)
+	{
+		FPlainPropsLinkerSandbox::MountPackageFromMemory(Package->GetFName());
+	}
+
+	UE_LOGFMT(LogPlainPropsUObject, Display, "Loading {SaveNum} packages...", Packages.Num());
+	for (UPackage* Package : Packages)
+	{
+		FPlainPropsLinkerSandbox::LoadPackageFromMemory(Package->GetFName());
+	}
+
+#if PP_WITH_ENGINE_HOOKS
+	SavePackageSettings.ExternalWriteExportsFunc = nullptr;
+	GSerializePropertiesOverride = nullptr;
+#endif
+
+	return 0;
+}
+#endif
+
+struct FPlainPropsBindingScope
+{
+	FPlainPropsBindingScope()
+		: Transform()
+		, Guid()
+		, Color()
+		, LinearColor()
+		, FieldPath(GUE.Structs.FieldPath)
+		, Delegate(GUE.Structs.Delegate)
+		, InlineMulticast({ GUE.Structs.MulticastInlineDelegate, GUE.Structs.MulticastDelegate })
+		, VerseFunction(GUE.Structs.VerseFunction)
+		, DynamicallyTypedValue(GUE.Structs.DynamicallyTypedValue)
+		, ReferencePropertyValue(GUE.Structs.ReferencePropertyValue)
+	{
+	}
+
 	TScopedStructBinding<FTransform, FDefaultRuntime> Transform;
 	TScopedStructBinding<FGuid, FDefaultRuntime> Guid;
 	TScopedStructBinding<FColor, FDefaultRuntime> Color;
 	TScopedStructBinding<FLinearColor, FDefaultRuntime> LinearColor;
-	TScopedStructBinding<FFieldPath, FDefaultRuntime> FieldPath(GUE.Structs.FieldPath);
-	TScopedStructBinding<FScriptDelegate, FDefaultRuntime> Delegate(GUE.Structs.Delegate);
+	TScopedStructBinding<FFieldPath, FDefaultRuntime> FieldPath;
+	TScopedStructBinding<FScriptDelegate, FDefaultRuntime> Delegate;
 	// MulticastDelegate declaration is shared with MulticastSparseDelegate
-	TScopedStructBinding<FMulticastScriptDelegate, FDefaultRuntime> InlineMulticast({GUE.Structs.MulticastInlineDelegate, GUE.Structs.MulticastDelegate});
+	TScopedStructBinding<FMulticastScriptDelegate, FDefaultRuntime> InlineMulticast;
 	// Verse
-	TScopedStructBinding<FVerseFunction, FDefaultRuntime> VerseFunction(GUE.Structs.VerseFunction);
-	TScopedStructBinding<::UE::FDynamicallyTypedValue, FDefaultRuntime> DynamicallyTypedValue(GUE.Structs.DynamicallyTypedValue);
-	TScopedStructBinding<FReferencePropertyValue, FDefaultRuntime> ReferencePropertyValue(GUE.Structs.ReferencePropertyValue);
+	TScopedStructBinding<FVerseFunction, FDefaultRuntime> VerseFunction;
+	TScopedStructBinding<::UE::FDynamicallyTypedValue, FDefaultRuntime> DynamicallyTypedValue;
+	TScopedStructBinding<FReferencePropertyValue, FDefaultRuntime> ReferencePropertyValue;
+};
+
+} // namespace PlainProps::UE
 
-	GUE.Defaults.BindZeroes(GUE.Structs.FieldPath, sizeof(FFieldPath), alignof(FFieldPath));
 
+static TUniquePtr<PlainProps::UE::FPlainPropsBindingScope> BindingScopePtr = nullptr;
+static TUniquePtr<PlainProps::DbgVis::FIdScope> DebugIdScopePtr = nullptr;
+
+void FPlainPropsUObjectModule::StartupModule()
+{
+	using namespace PlainProps;
+	using namespace PlainProps::UE;
+	BindingScopePtr = MakeUnique<FPlainPropsBindingScope>();
+	DebugIdScopePtr = MakeUnique<DbgVis::FIdScope>(GUE.Names, "SensName");
 	InitBatchedProperties();
-	BindInitialTypes();
-	return Roundtrip(Options);
 }
 
-} // namespace PlainProps::UE
+void FPlainPropsUObjectModule::ShutdownModule()
+{
+	BindingScopePtr.Reset();
+	DebugIdScopePtr.Reset();
+}
 
 #include "PlainPropsCommandlets.h"
 #include "AssetRegistry/AssetData.h"
@@ -4432,11 +5234,13 @@ int32 UTestPlainPropsCommandlet::Main(const FString& Params)
 {
 	using namespace PlainProps;
 	using namespace PlainProps::UE;
-	DbgVis::FIdScope _(GUE.Names, "SensName");
 
+	TArray<UObject*> Objects;
 	if (int32 LoadIdx = Params.Find("-load="); LoadIdx != INDEX_NONE)
 	{
-		// E.g. -run=TestPlainProps -load=/BRRoot/BRRoot.BRRoot,/Game/Maps/FrontEnd.FrontEnd
+		// E.g.
+		// -run=TestPlainProps -load=/BRRoot/BRRoot.BRRoot,/Game/Maps/FrontEnd.FrontEnd
+		// -run=TestPlainProps -load=/Game/Maps/Highrise.Highrise -save
 
 		const TCHAR* It = &Params[LoadIdx + 6];
 		FStringView Assets(It, FAsciiSet::FindFirstOrEnd(It, FAsciiSet(" ")) - It);
@@ -4445,11 +5249,17 @@ int32 UTestPlainPropsCommandlet::Main(const FString& Params)
 		int32 CommaIndex;
 		while (Assets.FindChar(',', CommaIndex))
 		{
-			FSoftObjectPath(Assets.Left(CommaIndex)).LoadAsync({});
+			FSoftObjectPath AssetPath(Assets.Left(CommaIndex)); 
+			UObject* Object = AssetPath.TryLoad();
+			check(Object);
+			Objects.Add(Object);
 			Assets.RightChopInline(CommaIndex + 1);
 		}
 
-		FSoftObjectPath(Assets).LoadAsync({});
+		FSoftObjectPath AssetPath(Assets);
+		UObject* Object = AssetPath.TryLoad();
+		check(Object);
+		Objects.Add(Object);
 	}
 	else if (Params.Find("-loadmaps") != INDEX_NONE)
 	{
@@ -4465,15 +5275,15 @@ int32 UTestPlainPropsCommandlet::Main(const FString& Params)
 		UE_LOGFMT(LogPlainPropsUObject, Display, "Loading all {Maps} maps...", Maps.Num());
 		for (FAssetData& Map : Maps)
 		{
-			Map.GetSoftObjectPath().LoadAsync({});
+			FSoftObjectPath AssetPath = Map.GetSoftObjectPath();
+			AssetPath.LoadAsync({});
 		}
 	}
 
 	FlushAsyncLoading();
-
 	UE_LOGFMT(LogPlainPropsUObject, Display, "Starting test...");
 
-	ERoundtrip Options = ERoundtrip::PP | ERoundtrip::UPS | ERoundtrip::TPS | ERoundtrip::TextMemory;
+	ERoundtrip Options = ERoundtrip::PP | ERoundtrip::UPS | ERoundtrip::TPS;
 	if (Params.Find("-pp") != INDEX_NONE)
 	{
 		Options = ERoundtrip::PP | ERoundtrip::TextMemory;
@@ -4482,9 +5292,20 @@ int32 UTestPlainPropsCommandlet::Main(const FString& Params)
 	{
 		Options = ERoundtrip::TextMemory | ERoundtrip::TextStable;
 	}
-	else if (Params.Find("-notext") != INDEX_NONE)
+
+#if WITH_EDITOR && PP_WITH_ENGINE_HOOKS
+	bool bSave = Params.Find("-save") != INDEX_NONE;
+	if (bSave)
 	{
-		EnumRemoveFlags(Options, ERoundtrip::TextMemory | ERoundtrip::TextStable);
+		RoundtripPackages(Objects, Options);
 	}
-	return TestBindings(Options);
-}
\ No newline at end of file
+	else
+#endif
+	{ 
+		RoundtripObjects(Objects, Options);
+	}
+	
+	FPlatformMisc::RequestExit(true, TEXT("TestPlainProps"));
+	UE_LOGFMT(LogPlainPropsUObject, Display, "Done!");
+	return 0;
+}
diff --git a/Engine/Plugins/Experimental/PlainPropsUObject/Source/Public/PlainPropsUObjectModule.h b/Engine/Plugins/Experimental/PlainPropsUObject/Source/Public/PlainPropsUObjectModule.h
--- a/Engine/Plugins/Experimental/PlainPropsUObject/Source/Public/PlainPropsUObjectModule.h
+++ b/Engine/Plugins/Experimental/PlainPropsUObject/Source/Public/PlainPropsUObjectModule.h
@@ -4,4 +4,8 @@
 
 #include "Modules/ModuleInterface.h"
 
-class FPlainPropsUObjectModule final : public IModuleInterface {};
\ No newline at end of file
+class FPlainPropsUObjectModule final : public IModuleInterface
+{
+	virtual void StartupModule();
+	virtual void ShutdownModule();
+};
\ No newline at end of file
