PlainProps: Remove UnrealEd dependency by replacing GEditor->SavePackage with UPackage::Save.

#rnx

diff --git a/Engine/Plugins/Experimental/PlainPropsUObject/Source/PlainPropsUObject.Build.cs b/Engine/Plugins/Experimental/PlainPropsUObject/Source/PlainPropsUObject.Build.cs
--- a/Engine/Plugins/Experimental/PlainPropsUObject/Source/PlainPropsUObject.Build.cs
+++ b/Engine/Plugins/Experimental/PlainPropsUObject/Source/PlainPropsUObject.Build.cs
@@ -23,7 +23,6 @@ public class PlainPropsUObject : ModuleRules
 		{
 			PrivateDependencyModuleNames.AddRange(
 				new string[] {
-					"UnrealEd", 		// for GEditor->SavePackage
 					"JsonUtilities", 	// for FJsonObjectConverter::UStructToJsonObjectString
 				}
 			);
diff --git a/Engine/Plugins/Experimental/PlainPropsUObject/Source/Private/PlainPropsUObjectBindings.cpp b/Engine/Plugins/Experimental/PlainPropsUObject/Source/Private/PlainPropsUObjectBindings.cpp
--- a/Engine/Plugins/Experimental/PlainPropsUObject/Source/Private/PlainPropsUObjectBindings.cpp
+++ b/Engine/Plugins/Experimental/PlainPropsUObject/Source/Private/PlainPropsUObjectBindings.cpp
@@ -55,10 +55,6 @@
 #include "UObject/VerseStringProperty.h"
 
 
-#if WITH_EDITOR
-#include "Editor.h"
-#endif
-
 #define PP_WITH_ENGINE_HOOKS 0
 
 using FUnicastScriptDelegate = TScriptDelegate<FNotThreadSafeNotCheckedDelegateMode>;
@@ -5143,7 +5139,7 @@ static int32 RoundtripPackages(TConstArrayView<UObject*> Objects, ERoundtrip Opt
 		FString TempFilename = FPaths::Combine(*FPaths::ProjectSavedDir(), TEXT("PlainProps"), Package->GetName() + Extension);
 		// FString ActualFilename = FPackageName::LongPackageNameToFilename(Asset->GetPackage()->GetName(), Extension);
 
-		GEditor->SavePackage(Package, Asset, *TempFilename, SaveArgs);
+		UPackage::Save(Package, Asset, *TempFilename, SaveArgs);
 	}
 
 	FPlainPropsLinkerSandbox::SavedPackages.Reserve(Packages.Num());
