PlainProps remove expensive World.h include

#rnx

diff --git a/Engine/Plugins/Experimental/PlainPropsUObject/Source/Private/PlainPropsUObjectBindings.cpp b/Engine/Plugins/Experimental/PlainPropsUObject/Source/Private/PlainPropsUObjectBindings.cpp
--- a/Engine/Plugins/Experimental/PlainPropsUObject/Source/Private/PlainPropsUObjectBindings.cpp
+++ b/Engine/Plugins/Experimental/PlainPropsUObject/Source/Private/PlainPropsUObjectBindings.cpp
@@ -5220,7 +5220,6 @@ void FPlainPropsUObjectModule::ShutdownModule()
 #include "PlainPropsCommandlets.h"
 #include "AssetRegistry/AssetData.h"
 #include "AssetRegistry/AssetRegistryModule.h"
-#include "Engine/World.h"
 
 UTestPlainPropsCommandlet::UTestPlainPropsCommandlet(const FObjectInitializer& Init)
 : Super(Init)
@@ -5266,7 +5265,7 @@ int32 UTestPlainPropsCommandlet::Main(const FString& Params)
 
 		UE_LOGFMT(LogPlainPropsUObject, Display, "Gathering all maps...");
 		TArray<FAssetData> Maps;
-		AssetRegistryModule.Get().GetAssetsByClass(UWorld::StaticClass()->GetClassPathName(), /* out */ Maps, true);
+		AssetRegistryModule.Get().GetAssetsByClass(FTopLevelAssetPath("/Script/Engine", "World"), /* out */ Maps, true);
 
 		UE_LOGFMT(LogPlainPropsUObject, Display, "Loading all {Maps} maps...", Maps.Num());
 		for (FAssetData& Map : Maps)
