PlainProps: Add class exclusions to package roundtripping to work around crashes due to non-implemented but required custom bindings in e.g. animation assets.

#rnx

diff --git a/Engine/Plugins/Experimental/PlainPropsUObject/Source/Private/PlainPropsUObjectBindings.cpp b/Engine/Plugins/Experimental/PlainPropsUObject/Source/Private/PlainPropsUObjectBindings.cpp
--- a/Engine/Plugins/Experimental/PlainPropsUObject/Source/Private/PlainPropsUObjectBindings.cpp
+++ b/Engine/Plugins/Experimental/PlainPropsUObject/Source/Private/PlainPropsUObjectBindings.cpp
@@ -4724,7 +4724,7 @@ struct FPlainPropsLinkerSandbox
 	{
 		FName PackageName = LinkerSave.LinkerRoot->GetFName();
 
-		UE_LOGFMT(LogPlainPropsUObject, Display, " Saving extension data for {PackageName}...", PackageName);
+		UE_LOGFMT(LogPlainPropsUObject, Log, " Saving PlainProps for {PackageName}...", PackageName);
 		FLinkerPropertyBatch Properties(LinkerSave);
 		TCustomPropertyBindings<FLinkerPropertyBatch> Customs(Properties, GUE.Customs);
 		FBatchSaver BatchSaver(Customs.Overlay, LinkerSave.ExportMap.Num());
@@ -4780,7 +4780,7 @@ struct FPlainPropsLinkerSandbox
 	{
 		if (TUniquePtr<FLinkerBatch>* LinkerBatch = LoadBatches.Find(PackageName))
 		{
-			UE_LOGFMT(LogPlainPropsUObject, Display, "  Loading {Package}...", PackageName);
+			UE_LOGFMT(LogPlainPropsUObject, Log, " Loading PlainProps for {Package}...", PackageName);
 			LoadPackage((*LinkerBatch)->Package, (*LinkerBatch)->PackagePath, LoadFlags);
 		}
 	}
@@ -4832,8 +4832,8 @@ TMap<FName, TUniquePtr<FPlainPropsLinkerSandbox::FLinkerBatch>> FPlainPropsLinke
 static int32 RoundtripObjects(TConstArrayView<UObject*> Objects, ERoundtrip Options)
 {
 	UE_LOGFMT(LogPlainPropsUObject, Display, "Binding types to PlainProps schemas...");
-	FTypeBinder StructBinder;
-	StructBinder.BindAllTypes();
+	FTypeBinder TypeBinder;
+	TypeBinder.BindAllTypes();
 
 	TArray<FInstance> Instances;
 	if (Objects.Num())
@@ -4841,10 +4841,10 @@ static int32 RoundtripObjects(TConstArrayView<UObject*> Objects, ERoundtrip Opti
 		UE_LOGFMT(LogPlainPropsUObject, Display, "Gathering UObjects in Assets...");
 		for (const UObject* Object : Objects)
 		{
-			ForEachObjectWithPackage(Object->GetPackage(), [&Instances,&StructBinder](UObject* Obj)
+			ForEachObjectWithPackage(Object->GetPackage(), [&Instances,&TypeBinder](UObject* Obj)
 			{
 				UClass* Class = Obj->GetClass();
-				if (StructBinder.ShouldBind(Class) && IncludeClass(Class))
+				if (TypeBinder.ShouldBind(Class) && IncludeClass(Class))
 				{
 					Instances.Emplace(Obj->GetPathName(), Obj);
 				}
@@ -4860,7 +4860,7 @@ static int32 RoundtripObjects(TConstArrayView<UObject*> Objects, ERoundtrip Opti
 		{
 			UObject* Object = *It;
 			UClass* Class = Object->GetClass();
-			if (StructBinder.ShouldBind(Class) && IncludeClass(Class))
+			if (TypeBinder.ShouldBind(Class) && IncludeClass(Class))
 			{
 				Instances.Emplace(Object->GetPathName(), Object);
 			}
@@ -5072,6 +5072,9 @@ static int32 RoundtripObjects(TConstArrayView<UObject*> Objects, ERoundtrip Opti
 #if WITH_EDITOR
 static int32 RoundtripPackages(TConstArrayView<UObject*> Objects, ERoundtrip Options)
 {
+	UE_LOGFMT(LogPlainPropsUObject, Display, "Binding types to PlainProps schemas...");
+	FTypeBinder{ .ExcludedPropertyFlags = CPF_Transient | CPF_Deprecated }.BindAllTypes();
+
 	TArray<UPackage*> Packages;
 	TArray<UObject*> AllPackages;
 
@@ -5100,26 +5103,30 @@ static int32 RoundtripPackages(TConstArrayView<UObject*> Objects, ERoundtrip Opt
 		{
 			continue;
 		}
-		bool bContainsBP = false;
-		ForEachObjectWithPackage(Package, [&bContainsBP](UObject* Object)
+		bool bExclude = false;
+		ForEachObjectWithPackage(Package, [&bExclude](UObject* Object)
 		{
-			if (Object->GetClass()->HasAnyClassFlags(CLASS_CompiledFromBlueprint))
+			UClass* Class = Object->GetClass();
+			if (Class->HasAnyClassFlags(CLASS_CompiledFromBlueprint))
+			{
+				bExclude = true;
+				return false;
+			}
+			if (!IncludeClass(Class))
 			{
-				bContainsBP = true;
+				UE_LOGFMT(LogPlainPropsUObject, Log, " Excluding PlainProps for {Package}", Object->GetPackage()->GetName());
+				bExclude = true;
 				return false;
 			}
 			return true;
 		});
-		if (bContainsBP)
+		if (bExclude)
 		{
 			continue;
 		}
 		Packages.Add(Package);
 	}
 
-	UE_LOGFMT(LogPlainPropsUObject, Display, "Binding types to PlainProps schemas...");
-	FTypeBinder{ .ExcludedPropertyFlags = CPF_Transient | CPF_Deprecated }.BindAllTypes();
-
 #if PP_WITH_ENGINE_HOOKS
 	FSavePackageSettings& SavePackageSettings = FSavePackageSettings::GetDefaultSettings();
 	SavePackageSettings.ExternalWriteExportsFunc = &FPlainPropsLinkerSandbox::SavePlainProps;
