PlainProps add new Engine plugin and remove engine dependency from PlainPropsUObject module

#rnx

diff --git a/Engine/Plugins/Experimental/PlainPropsEngine/PlainPropsEngine.uplugin b/Engine/Plugins/Experimental/PlainPropsEngine/PlainPropsEngine.uplugin
new file mode 100644
--- /dev/null
+++ b/Engine/Plugins/Experimental/PlainPropsEngine/PlainPropsEngine.uplugin
@@ -0,0 +1,31 @@
+{
+	"FileVersion": 3,
+	"FriendlyName": "PlainPropsEngine",
+	"Description": "New Serialization Stack Prototype - Engine Bindings",
+	"Category": "Experimental",
+	"CreatedBy": "Epic Games, Inc.",
+	"CreatedByURL": "https://epicgames.com",
+	"EnabledByDefault": false,
+	"IsExperimentalVersion": true,
+	"SupportedTargetPlatforms": [
+		"Win64"
+	],
+	"Modules": [
+		{
+			"Name": "PlainPropsEngine",
+			"Type": "Runtime",
+			"LoadingPhase": "PreLoadingScreen",
+			"PlatformAllowList": [ "Win64" ],
+		}
+	],
+	"Plugins": [
+		{
+			"Name": "PlainProps",
+			"Enabled": true
+		},
+		{
+			"Name": "PlainPropsUObject",
+			"Enabled": true
+		}
+	]
+}
diff --git a/Engine/Plugins/Experimental/PlainPropsEngine/Source/PlainPropsEngine.Build.cs b/Engine/Plugins/Experimental/PlainPropsEngine/Source/PlainPropsEngine.Build.cs
new file mode 100644
--- /dev/null
+++ b/Engine/Plugins/Experimental/PlainPropsEngine/Source/PlainPropsEngine.Build.cs
@@ -0,0 +1,23 @@
+// Copyright Epic Games, Inc. All Rights Reserved.
+
+using UnrealBuildTool;
+
+public class PlainPropsEngine : ModuleRules
+{
+	public PlainPropsEngine(ReadOnlyTargetRules Target) : base(Target)
+	{
+		bDisableStaticAnalysis = true;
+
+		CppStandard = CppStandardVersion.Cpp20;
+
+		PrivateDependencyModuleNames.AddRange(
+			new string[] {
+				"Core",
+				"CoreUObject",
+				"CorePreciseFP",
+				"Engine",
+				"PlainProps",
+				"PlainPropsUObject"
+		});
+	}
+}
diff --git a/Engine/Plugins/Experimental/PlainPropsEngine/Source/Private/PlainPropsCommandlets.cpp b/Engine/Plugins/Experimental/PlainPropsEngine/Source/Private/PlainPropsCommandlets.cpp
new file mode 100644
--- /dev/null
+++ b/Engine/Plugins/Experimental/PlainPropsEngine/Source/Private/PlainPropsCommandlets.cpp
@@ -0,0 +1,108 @@
+// Copyright Epic Games, Inc. All Rights Reserved.
+
+#include "PlainPropsCommandlets.h"
+#include "PlainPropsEngineModule.h"
+#include "PlainPropsRoundtripTest.h"
+#include "AssetRegistry/AssetData.h"
+#include "AssetRegistry/AssetRegistryModule.h"
+#include "Misc/AsciiSet.h"
+#include "Misc/CommandLine.h"
+#include "Misc/PackageName.h"
+
+UTestPlainPropsCommandlet::UTestPlainPropsCommandlet(const FObjectInitializer& Init)
+: Super(Init)
+{}
+
+int32 UTestPlainPropsCommandlet::Main(const FString& Params)
+{
+	using namespace PlainProps;
+	using namespace PlainProps::UE;
+
+	TArray<UObject*> Objects;
+	if (int32 LoadIdx = Params.Find("-load="); LoadIdx != INDEX_NONE)
+	{
+		// E.g.
+		// -run=TestPlainProps -load=/BRRoot/BRRoot.BRRoot,/Game/Maps/FrontEnd.FrontEnd
+		// -run=TestPlainProps -load=/Game/Maps/Highrise.Highrise -save
+
+		const TCHAR* It = &Params[LoadIdx + 6];
+		FStringView Assets(It, FAsciiSet::FindFirstOrEnd(It, FAsciiSet(" ")) - It);
+		UE_LOGFMT(LogPlainPropsEngine, Display, "Loading {Assets}...", Assets);
+
+		int32 CommaIndex;
+		while (Assets.FindChar(',', CommaIndex))
+		{
+			FSoftObjectPath AssetPath(Assets.Left(CommaIndex)); 
+			UObject* Object = AssetPath.TryLoad();
+			check(Object);
+			Objects.Add(Object);
+			Assets.RightChopInline(CommaIndex + 1);
+		}
+
+		FSoftObjectPath AssetPath(Assets);
+		UObject* Object = AssetPath.TryLoad();
+		check(Object);
+		Objects.Add(Object);
+	}
+	else if (Params.Find("-loadmaps") != INDEX_NONE)
+	{
+		// load all .umaps in asset registry
+		UE_LOGFMT(LogPlainPropsEngine, Display, "Loading asset registry...");
+		FAssetRegistryModule& AssetRegistryModule = FModuleManager::LoadModuleChecked<FAssetRegistryModule>(AssetRegistryConstants::ModuleName);
+		AssetRegistryModule.Get().SearchAllAssets(true);
+
+		UE_LOGFMT(LogPlainPropsEngine, Display, "Gathering all maps...");
+		TArray<FAssetData> Maps;
+		AssetRegistryModule.Get().GetAssetsByClass(FTopLevelAssetPath("/Script/Engine", "World"), /* out */ Maps, true);
+
+		UE_LOGFMT(LogPlainPropsEngine, Display, "Loading all {Maps} maps...", Maps.Num());
+		for (FAssetData& Map : Maps)
+		{
+			FSoftObjectPath AssetPath = Map.GetSoftObjectPath();
+			AssetPath.LoadAsync({});
+		}
+	}
+
+	FlushAsyncLoading();
+	UE_LOGFMT(LogPlainPropsEngine, Display, "Starting test...");
+
+	// Parse bind mode
+	EPropertyFlags Skip = CPF_None;
+	if (Params.Find("-bind=source") != INDEX_NONE)
+	{
+		Skip = CPF_Transient | CPF_Deprecated | CPF_SkipSerialization;
+	}
+	else if (Params.Find("-bind=runtime") != INDEX_NONE)
+	{
+		Skip = CPF_Transient | CPF_Deprecated | CPF_SkipSerialization | CPF_EditorOnly;
+		// todo: add engine specific custom cook bindings here
+	}
+
+	UE_LOGFMT(LogPlainPropsEngine, Display, "Binding types to PlainProps schemas...");
+	BindAllTypes();
+
+	// Parse roundtrip options
+	ERoundtrip Options = ERoundtrip::PP | ERoundtrip::UPS | ERoundtrip::TPS;
+	if (Params.Find("-pp") != INDEX_NONE)
+	{
+		Options = ERoundtrip::PP | ERoundtrip::TextMemory;
+	}
+	else if (Params.Find("-text") != INDEX_NONE)
+	{
+		Options = ERoundtrip::TextMemory | ERoundtrip::TextStable;
+	}
+
+	// Roundtrip
+	if (Params.Find("-package") != INDEX_NONE)
+	{
+		RoundtripViaPackages(Objects, Options);
+	}
+	else
+	{ 
+		RoundtripViaBatch(Objects, Options);
+	}
+	
+	FPlatformMisc::RequestExit(true, TEXT("TestPlainProps"));
+	UE_LOGFMT(LogPlainPropsEngine, Display, "Done!");
+	return 0;
+}
diff --git a/Engine/Plugins/Experimental/PlainPropsEngine/Source/Private/PlainPropsEngineBindings.cpp b/Engine/Plugins/Experimental/PlainPropsEngine/Source/Private/PlainPropsEngineBindings.cpp
new file mode 100644
--- /dev/null
+++ b/Engine/Plugins/Experimental/PlainPropsEngine/Source/Private/PlainPropsEngineBindings.cpp
@@ -0,0 +1,8 @@
+// Copyright Epic Games, Inc. All Rights Reserved.
+
+#include "PlainPropsEngineModule.h"
+
+namespace PlainProps::UE
+{
+
+} // namespace PlainProps::UE
diff --git a/Engine/Plugins/Experimental/PlainPropsEngine/Source/Private/PlainPropsEngineModule.cpp b/Engine/Plugins/Experimental/PlainPropsEngine/Source/Private/PlainPropsEngineModule.cpp
new file mode 100644
--- /dev/null
+++ b/Engine/Plugins/Experimental/PlainPropsEngine/Source/Private/PlainPropsEngineModule.cpp
@@ -0,0 +1,11 @@
+// Copyright Epic Games, Inc. All Rights Reserved.
+
+#include "PlainPropsEngineModule.h"
+#include "Modules/ModuleManager.h"
+
+DEFINE_LOG_CATEGORY(LogPlainPropsEngine);
+
+IMPLEMENT_MODULE(FPlainPropsEngineModule, PlainPropsEngine);
+
+void FPlainPropsEngineModule::StartupModule() {}
+void FPlainPropsEngineModule::ShutdownModule() {}
\ No newline at end of file
diff --git a/Engine/Plugins/Experimental/PlainPropsEngine/Source/Public/PlainPropsCommandlets.h b/Engine/Plugins/Experimental/PlainPropsEngine/Source/Public/PlainPropsCommandlets.h
new file mode 100644
--- /dev/null
+++ b/Engine/Plugins/Experimental/PlainPropsEngine/Source/Public/PlainPropsCommandlets.h
@@ -0,0 +1,15 @@
+// Copyright Epic Games, Inc. All Rights Reserved.
+
+#pragma once
+
+#include "CoreMinimal.h"
+#include "UObject/ObjectMacros.h"
+#include "Commandlets/Commandlet.h"
+#include "PlainPropsCommandlets.generated.h"
+
+UCLASS()
+class UTestPlainPropsCommandlet : public UCommandlet
+{
+	GENERATED_UCLASS_BODY()
+	virtual int32 Main(const FString& Params) override;
+};
diff --git a/Engine/Plugins/Experimental/PlainPropsEngine/Source/Public/PlainPropsEngineModule.h b/Engine/Plugins/Experimental/PlainPropsEngine/Source/Public/PlainPropsEngineModule.h
new file mode 100644
--- /dev/null
+++ b/Engine/Plugins/Experimental/PlainPropsEngine/Source/Public/PlainPropsEngineModule.h
@@ -0,0 +1,14 @@
+// Copyright Epic Games, Inc. All Rights Reserved.
+
+#pragma once
+
+#include "Modules/ModuleInterface.h"
+#include "Logging/StructuredLog.h"
+
+DECLARE_LOG_CATEGORY_EXTERN(LogPlainPropsEngine, Log, All);
+
+class FPlainPropsEngineModule final : public IModuleInterface
+{
+	virtual void StartupModule();
+	virtual void ShutdownModule();
+};
\ No newline at end of file
diff --git a/Engine/Plugins/Experimental/PlainPropsUObject/Source/PlainPropsUObject.Build.cs b/Engine/Plugins/Experimental/PlainPropsUObject/Source/PlainPropsUObject.Build.cs
--- a/Engine/Plugins/Experimental/PlainPropsUObject/Source/PlainPropsUObject.Build.cs
+++ b/Engine/Plugins/Experimental/PlainPropsUObject/Source/PlainPropsUObject.Build.cs
@@ -15,17 +15,8 @@ public class PlainPropsUObject : ModuleRules
 				"Core",
 				"CoreUObject",
 				"CorePreciseFP",
-				"Engine",
+				"JsonUtilities", 	// for FJsonObjectConverter::UStructToJsonObjectString
 				"PlainProps"
 		});
-
-		if (Target.bBuildEditor == true)
-		{
-			PrivateDependencyModuleNames.AddRange(
-				new string[] {
-					"JsonUtilities", 	// for FJsonObjectConverter::UStructToJsonObjectString
-				}
-			);
-		}
 	}
 }
diff --git a/Engine/Plugins/Experimental/PlainPropsUObject/Source/Private/PlainPropsUObjectBindings.cpp b/Engine/Plugins/Experimental/PlainPropsUObject/Source/Private/PlainPropsUObjectBindings.cpp
--- a/Engine/Plugins/Experimental/PlainPropsUObject/Source/Private/PlainPropsUObjectBindings.cpp
+++ b/Engine/Plugins/Experimental/PlainPropsUObject/Source/Private/PlainPropsUObjectBindings.cpp
@@ -1,13 +1,14 @@
 // Copyright Epic Games, Inc. All Rights Reserved.
 
+#include "PlainPropsRoundtripTest.h"
 #include "PlainPropsUObjectRuntime.h"
-#include "PlainPropsUeCoreBindings.h"
+#include "PlainPropsUObjectModule.h"
 #include "PlainPropsBuildSchema.h"
 #include "PlainPropsDiff.h"
 #include "PlainPropsParse.h"
 #include "PlainPropsPrint.h"
 #include "PlainPropsRead.h"
-#include "PlainPropsUObjectModule.h"
+#include "PlainPropsUeCoreBindings.h"
 #include "PlainPropsVisualize.h"
 #include "PlainPropsWrite.h"
 #include "Algo/Compare.h"
@@ -952,17 +953,6 @@ static bool DenyMaxValue(FName Enum)
 	return !Algo::Find(AllowsMax, Enum);
 }
 
-enum class ERoundtrip : uint8
-{
-	None 		= 0,
-	PP 			= 1 << 0,
-	TPS 		= 1 << 1,
-	UPS 		= 1 << 2,
-	TextMemory	= 1 << 3,
-	TextStable	= 1 << 4,
-};
-ENUM_CLASS_FLAGS(ERoundtrip);
-
 static FEnumId DeclareEnum(UEnum* Enum)
 {
 	FType Type = IndexType(Enum);
@@ -3406,136 +3396,140 @@ FStructDeclarationPtr FPropertyBinder::Declare() const
 
 ////////////////////////////////////////////////////////////////////////////////////////////////
 
-struct FTypeBinder
-{
-	EPropertyFlags ExcludedPropertyFlags = EPropertyFlags::CPF_None;
+static EPropertyFlags GDontBind;
 
-	inline bool ShouldBind(const FProperty* Property)
-	{
-		if (Property->PropertyFlags & ExcludedPropertyFlags)
-		{
-			return false;
-		}
+static bool ShouldBind(const UStruct* Struct);
 
-		if (HasAny<CASTCLASS_FStructProperty>(Property->GetCastFlags()))
-		{
-			return ShouldBind(static_cast<const FStructProperty*>(Property)->Struct);
-		}
+static bool ShouldBind(const FProperty* Property)
+{
+	if (Property->PropertyFlags & GDontBind)
+	{
+		return false;
+	}
 
-		return true;
+	if (HasAny<CASTCLASS_FStructProperty>(Property->GetCastFlags()))
+	{
+		return ShouldBind(static_cast<const FStructProperty*>(Property)->Struct);
 	}
 
-	bool ShouldBind(const UStruct* Struct)
+	return true;
+}
+
+static bool ShouldBind(const UStruct* Struct)
+{
+	for (FProperty* Property = Struct->PropertyLink; Property; Property = Property->PropertyLinkNext)
 	{
-		for (FProperty* Property = Struct->PropertyLink; Property; Property = Property->PropertyLinkNext)
+		if (ShouldBind(Property))
 		{
-			if (ShouldBind(Property))
-			{
-				return true;
-			}
+			return true;
 		}
-		return false;
 	}
+	return false;
+}
 
-	inline void BindMembers(FPropertyBinder& Out, const UStruct* Struct)
+inline void BindMembers(FPropertyBinder& Out, const UStruct* Struct)
+{
+	for (FProperty* It = Struct->PropertyLink; It && It->GetOwner<UStruct>() == Struct; It = It->PropertyLinkNext)
 	{
-		for (FProperty* It = Struct->PropertyLink; It && It->GetOwner<UStruct>() == Struct; It = It->PropertyLinkNext)
+		if (ShouldBind(It))
 		{
-			if (ShouldBind(It))
-			{
-				Out.BindMember(It);
-			}
+			Out.BindMember(It);
 		}
 	}
+}
 
-	void BindSuperMembers(FPropertyBinder& Out, const UStruct* Struct)
+static void BindSuperMembers(FPropertyBinder& Out, const UStruct* Struct)
+{
+	if (const UStruct* Super = Struct->GetInheritanceSuper())
 	{
-		if (const UStruct* Super = Struct->GetInheritanceSuper())
+		BindSuperMembers(Out, Super);
+		if (ShouldBind(Super))
 		{
-			BindSuperMembers(Out, Super);
-			if (ShouldBind(Super))
-			{
-				BindMembers(Out, Super);
-			}
+			BindMembers(Out, Super);
 		}
 	}
+}
 
-	const UStruct* GetSuperToBind(const UStruct* Struct)
+static const UStruct* GetSuperToBind(const UStruct* Struct)
+{
+	const UStruct* Super = Struct->GetInheritanceSuper();
+	while (Super)
 	{
-		const UStruct* Super = Struct->GetInheritanceSuper();
-		while (Super)
+		if (ShouldBind(Super))
 		{
-			if (ShouldBind(Super))
-			{
-				return SkipEmptyBases(Super);
-			}
-			Super = Super->GetInheritanceSuper();
+			return SkipEmptyBases(Super);
 		}
-		return nullptr;
+		Super = Super->GetInheritanceSuper();
 	}
+	return nullptr;
+}
 
-	void BindStruct(FBindId Id, const UStruct* Struct)
+static void BindStruct(FBindId Id, const UStruct* Struct)
+{
+	if (GUE.Customs.FindStruct(Id))
 	{
-		if (GUE.Customs.FindStruct(Id))
-		{
-			return;
-		}
+		return;
+	}
 
-		if (GUE.Schemas.FindStruct(Id))
-		{
-			return;
-		}
+	if (GUE.Schemas.FindStruct(Id))
+	{
+		return;
+	}
 
-		FPropertyBinder Binder(Id, GetOccupancy(Struct));
-		if (const UStruct* Super = GetSuperToBind(Struct))
+	FPropertyBinder Binder(Id, GetOccupancy(Struct));
+	if (const UStruct* Super = GetSuperToBind(Struct))
+	{
+		if (Binder.IsDense())
 		{
-			if (Binder.IsDense())
+			// Flatten inheritance chain for dense structs
+			for (; Super; Super = GetSuperToBind(Super))
 			{
-				// Flatten inheritance chain for dense structs
-				for (; Super; Super = GetSuperToBind(Super))
-				{
-					BindMembers(/* out */ Binder, Super);	
-				}
-			}
-			else
-			{
-				FDeclId SuperId = GUE.Names.IndexDeclId(IndexType(Super));
-				Binder.BindSuper(SuperId);
+				BindMembers(/* out */ Binder, Super);	
 			}
 		}
+		else
+		{
+			FDeclId SuperId = GUE.Names.IndexDeclId(IndexType(Super));
+			Binder.BindSuper(SuperId);
+		}
+	}
 
-		BindMembers(/* out */ Binder, Struct);
+	BindMembers(/* out */ Binder, Struct);
 
-		GUE.Schemas.BindStruct(Id, Binder.GetMembers(), Binder.Declare());
+	GUE.Schemas.BindStruct(Id, Binder.GetMembers(), Binder.Declare());
 		
-		// Don't bind CDO defaults, object defaults are passed in from top and objects aren't owned by containers 
-		if (Struct->HasAnyCastFlags(CASTCLASS_UScriptStruct))
-		{
-			GUE.Defaults.Bind(Id, static_cast<const UScriptStruct*>(Struct));
-		}
+	// Don't bind CDO defaults, object defaults are passed in from top and objects aren't owned by containers 
+	if (Struct->HasAnyCastFlags(CASTCLASS_UScriptStruct))
+	{
+		GUE.Defaults.Bind(Id, static_cast<const UScriptStruct*>(Struct));
 	}
+}
 
-	void BindAllTypes()
+void DontBind(EPropertyFlags Skip)
+{
+	GDontBind = Skip;
+}
+
+void BindAllTypes()
+{
+	// Declare all UEnums
+	for (TObjectIterator<UEnum> It; It; ++It)
 	{
-		// Declare all UEnums
-		for (TObjectIterator<UEnum> It; It; ++It)
-		{
-			DeclareEnum(*It);
-		}
+		DeclareEnum(*It);
+	}
 
-		// Bind all UScriptStruct/UClass/UFunction
-		static const FBindId SkipStructs[] = { GUE.Structs.VerseFunction };
-		for (TObjectIterator<UStruct> It; It; ++It)
+	// Bind all UScriptStruct/UClass/UFunction
+	static const FBindId SkipStructs[] = { GUE.Structs.VerseFunction };
+	for (TObjectIterator<UStruct> It; It; ++It)
+	{
+		const UStruct* Struct = *It;
+		FBindId Id = GUE.Names.IndexBindId(IndexType(Struct));
+		if (!Algo::Find(SkipStructs, Id))
 		{
-			const UStruct* Struct = *It;
-			FBindId Id = GUE.Names.IndexBindId(IndexType(Struct));
-			if (!Algo::Find(SkipStructs, Id))
-			{
-				BindStruct(Id, Struct);
-			}
+			BindStruct(Id, Struct);
 		}
 	}
-};
+}
 
 ////////////////////////////////////////////////////////////////////////////////////////////////
 
@@ -4829,22 +4823,18 @@ TMap<FName, TUniquePtr<FPlainPropsLinkerSandbox::FLinkerBatch>> FPlainPropsLinke
 
 ////////////////////////////////////////////////////////////////////////////////
 
-static int32 RoundtripObjects(TConstArrayView<UObject*> Objects, ERoundtrip Options)
+int32 RoundtripViaBatch(TConstArrayView<UObject*> Objects, ERoundtrip Options)
 {
-	UE_LOGFMT(LogPlainPropsUObject, Display, "Binding types to PlainProps schemas...");
-	FTypeBinder TypeBinder;
-	TypeBinder.BindAllTypes();
-
 	TArray<FInstance> Instances;
 	if (Objects.Num())
 	{
 		UE_LOGFMT(LogPlainPropsUObject, Display, "Gathering UObjects in Assets...");
 		for (const UObject* Object : Objects)
 		{
-			ForEachObjectWithPackage(Object->GetPackage(), [&Instances,&TypeBinder](UObject* Obj)
+			ForEachObjectWithPackage(Object->GetPackage(), [&Instances](UObject* Obj)
 			{
 				UClass* Class = Obj->GetClass();
-				if (TypeBinder.ShouldBind(Class) && IncludeClass(Class))
+				if (ShouldBind(Class) && IncludeClass(Class))
 				{
 					Instances.Emplace(Obj->GetPathName(), Obj);
 				}
@@ -4860,7 +4850,7 @@ static int32 RoundtripObjects(TConstArrayView<UObject*> Objects, ERoundtrip Opti
 		{
 			UObject* Object = *It;
 			UClass* Class = Object->GetClass();
-			if (TypeBinder.ShouldBind(Class) && IncludeClass(Class))
+			if (ShouldBind(Class) && IncludeClass(Class))
 			{
 				Instances.Emplace(Object->GetPathName(), Object);
 			}
@@ -5069,12 +5059,8 @@ static int32 RoundtripObjects(TConstArrayView<UObject*> Objects, ERoundtrip Opti
 
 ////////////////////////////////////////////////////////////////////////////////
 
-#if WITH_EDITOR
-static int32 RoundtripPackages(TConstArrayView<UObject*> Objects, ERoundtrip Options)
+int32 RoundtripViaPackages(TConstArrayView<UObject*> Objects, ERoundtrip Options)
 {
-	UE_LOGFMT(LogPlainPropsUObject, Display, "Binding types to PlainProps schemas...");
-	FTypeBinder{ .ExcludedPropertyFlags = CPF_Transient | CPF_Deprecated }.BindAllTypes();
-
 	TArray<UPackage*> Packages;
 	TArray<UObject*> AllPackages;
 
@@ -5128,6 +5114,7 @@ static int32 RoundtripPackages(TConstArrayView<UObject*> Objects, ERoundtrip Opt
 	}
 
 #if PP_WITH_ENGINE_HOOKS
+	// Todo consider TGuardValue
 	FSavePackageSettings& SavePackageSettings = FSavePackageSettings::GetDefaultSettings();
 	SavePackageSettings.ExternalWriteExportsFunc = &FPlainPropsLinkerSandbox::SavePlainProps;
 	SavePackageSettings.ExternalWriteFileFunc = &FPlainPropsLinkerSandbox::SavePackageToMemory;
@@ -5149,6 +5136,7 @@ static int32 RoundtripPackages(TConstArrayView<UObject*> Objects, ERoundtrip Opt
 		UPackage::Save(Package, Asset, *TempFilename, SaveArgs);
 	}
 
+#if WITH_EDITOR
 	FPlainPropsLinkerSandbox::SavedPackages.Reserve(Packages.Num());
 	FPlainPropsLinkerSandbox::LoadBatches.Reserve(Packages.Num());
 
@@ -5163,6 +5151,7 @@ static int32 RoundtripPackages(TConstArrayView<UObject*> Objects, ERoundtrip Opt
 	{
 		FPlainPropsLinkerSandbox::LoadPackageFromMemory(Package->GetFName());
 	}
+#endif
 
 #if PP_WITH_ENGINE_HOOKS
 	SavePackageSettings.ExternalWriteExportsFunc = nullptr;
@@ -5171,7 +5160,6 @@ static int32 RoundtripPackages(TConstArrayView<UObject*> Objects, ERoundtrip Opt
 
 	return 0;
 }
-#endif
 
 struct FPlainPropsBindingScope
 {
@@ -5205,7 +5193,6 @@ struct FPlainPropsBindingScope
 
 } // namespace PlainProps::UE
 
-
 static TUniquePtr<PlainProps::UE::FPlainPropsBindingScope> BindingScopePtr = nullptr;
 static TUniquePtr<PlainProps::DbgVis::FIdScope> DebugIdScopePtr = nullptr;
 
@@ -5222,92 +5209,4 @@ void FPlainPropsUObjectModule::ShutdownModule()
 {
 	BindingScopePtr.Reset();
 	DebugIdScopePtr.Reset();
-}
-
-#include "PlainPropsCommandlets.h"
-#include "AssetRegistry/AssetData.h"
-#include "AssetRegistry/AssetRegistryModule.h"
-
-UTestPlainPropsCommandlet::UTestPlainPropsCommandlet(const FObjectInitializer& Init)
-: Super(Init)
-{}
-
-int32 UTestPlainPropsCommandlet::Main(const FString& Params)
-{
-	using namespace PlainProps;
-	using namespace PlainProps::UE;
-
-	TArray<UObject*> Objects;
-	if (int32 LoadIdx = Params.Find("-load="); LoadIdx != INDEX_NONE)
-	{
-		// E.g.
-		// -run=TestPlainProps -load=/BRRoot/BRRoot.BRRoot,/Game/Maps/FrontEnd.FrontEnd
-		// -run=TestPlainProps -load=/Game/Maps/Highrise.Highrise -save
-
-		const TCHAR* It = &Params[LoadIdx + 6];
-		FStringView Assets(It, FAsciiSet::FindFirstOrEnd(It, FAsciiSet(" ")) - It);
-		UE_LOGFMT(LogPlainPropsUObject, Display, "Loading {Assets}...", Assets);
-
-		int32 CommaIndex;
-		while (Assets.FindChar(',', CommaIndex))
-		{
-			FSoftObjectPath AssetPath(Assets.Left(CommaIndex)); 
-			UObject* Object = AssetPath.TryLoad();
-			check(Object);
-			Objects.Add(Object);
-			Assets.RightChopInline(CommaIndex + 1);
-		}
-
-		FSoftObjectPath AssetPath(Assets);
-		UObject* Object = AssetPath.TryLoad();
-		check(Object);
-		Objects.Add(Object);
-	}
-	else if (Params.Find("-loadmaps") != INDEX_NONE)
-	{
-		// load all .umaps in asset registry
-		UE_LOGFMT(LogPlainPropsUObject, Display, "Loading asset registry...");
-		FAssetRegistryModule& AssetRegistryModule = FModuleManager::LoadModuleChecked<FAssetRegistryModule>(AssetRegistryConstants::ModuleName);
-		AssetRegistryModule.Get().SearchAllAssets(true);
-
-		UE_LOGFMT(LogPlainPropsUObject, Display, "Gathering all maps...");
-		TArray<FAssetData> Maps;
-		AssetRegistryModule.Get().GetAssetsByClass(FTopLevelAssetPath("/Script/Engine", "World"), /* out */ Maps, true);
-
-		UE_LOGFMT(LogPlainPropsUObject, Display, "Loading all {Maps} maps...", Maps.Num());
-		for (FAssetData& Map : Maps)
-		{
-			FSoftObjectPath AssetPath = Map.GetSoftObjectPath();
-			AssetPath.LoadAsync({});
-		}
-	}
-
-	FlushAsyncLoading();
-	UE_LOGFMT(LogPlainPropsUObject, Display, "Starting test...");
-
-	ERoundtrip Options = ERoundtrip::PP | ERoundtrip::UPS | ERoundtrip::TPS;
-	if (Params.Find("-pp") != INDEX_NONE)
-	{
-		Options = ERoundtrip::PP | ERoundtrip::TextMemory;
-	}
-	else if (Params.Find("-text") != INDEX_NONE)
-	{
-		Options = ERoundtrip::TextMemory | ERoundtrip::TextStable;
-	}
-
-#if WITH_EDITOR && PP_WITH_ENGINE_HOOKS
-	bool bSave = Params.Find("-save") != INDEX_NONE;
-	if (bSave)
-	{
-		RoundtripPackages(Objects, Options);
-	}
-	else
-#endif
-	{ 
-		RoundtripObjects(Objects, Options);
-	}
-	
-	FPlatformMisc::RequestExit(true, TEXT("TestPlainProps"));
-	UE_LOGFMT(LogPlainPropsUObject, Display, "Done!");
-	return 0;
-}
+}
\ No newline at end of file
diff --git a/Engine/Plugins/Experimental/PlainPropsUObject/Source/Public/PlainPropsCommandlets.h b/Engine/Plugins/Experimental/PlainPropsUObject/Source/Public/PlainPropsCommandlets.h
deleted file mode 100644
diff --git a/Engine/Plugins/Experimental/PlainPropsUObject/Source/Public/PlainPropsRoundtripTest.h b/Engine/Plugins/Experimental/PlainPropsUObject/Source/Public/PlainPropsRoundtripTest.h
new file mode 100644
--- /dev/null
+++ b/Engine/Plugins/Experimental/PlainPropsUObject/Source/Public/PlainPropsRoundtripTest.h
@@ -0,0 +1,29 @@
+// Copyright Epic Games, Inc. All Rights Reserved.
+
+#pragma once
+
+#include "CoreTypes.h"
+#include "Misc/EnumClassFlags.h"
+#include "UObject/ObjectMacros.h"
+
+namespace PlainProps::UE
+{
+
+enum class ERoundtrip : uint8
+{
+	None 		= 0,
+	PP 			= 1 << 0,
+	TPS 		= 1 << 1,
+	UPS 		= 1 << 2,
+	TextMemory	= 1 << 3,
+	TextStable	= 1 << 4,
+};
+ENUM_CLASS_FLAGS(ERoundtrip);
+
+
+PLAINPROPSUOBJECT_API void DontBind(EPropertyFlags Skip);
+PLAINPROPSUOBJECT_API void BindAllTypes();
+PLAINPROPSUOBJECT_API int32 RoundtripViaBatch(TConstArrayView<UObject*> Objects, ERoundtrip Options);
+PLAINPROPSUOBJECT_API int32 RoundtripViaPackages(TConstArrayView<UObject*> Objects, ERoundtrip Options);
+
+}
