PlainProps preparing to add custom bindings for Engine types

#rnx

diff --git a/Engine/Plugins/Experimental/PlainPropsEngine/Source/Private/PlainPropsCommandlets.cpp b/Engine/Plugins/Experimental/PlainPropsEngine/Source/Private/PlainPropsCommandlets.cpp
--- a/Engine/Plugins/Experimental/PlainPropsEngine/Source/Private/PlainPropsCommandlets.cpp
+++ b/Engine/Plugins/Experimental/PlainPropsEngine/Source/Private/PlainPropsCommandlets.cpp
@@ -1,7 +1,7 @@
 // Copyright Epic Games, Inc. All Rights Reserved.
 
 #include "PlainPropsCommandlets.h"
-#include "PlainPropsEngineModule.h"
+#include "PlainPropsEngineBindings.h"
 #include "PlainPropsRoundtripTest.h"
 #include "AssetRegistry/AssetData.h"
 #include "AssetRegistry/AssetRegistryModule.h"
@@ -66,20 +66,20 @@ int32 UTestPlainPropsCommandlet::Main(const FString& Params)
 	FlushAsyncLoading();
 	UE_LOGFMT(LogPlainPropsEngine, Display, "Starting test...");
 
-	// Parse bind mode
-	EPropertyFlags Skip = CPF_None;
+	UE_LOGFMT(LogPlainPropsEngine, Display, "Binding types to PlainProps schemas...");
+	
+	EBindMode Mode = EBindMode::All;
 	if (Params.Find("-bind=source") != INDEX_NONE)
 	{
-		Skip = CPF_Transient | CPF_Deprecated | CPF_SkipSerialization;
+		Mode = EBindMode::Source;
 	}
 	else if (Params.Find("-bind=runtime") != INDEX_NONE)
 	{
-		Skip = CPF_Transient | CPF_Deprecated | CPF_SkipSerialization | CPF_EditorOnly;
-		// todo: add engine specific custom cook bindings here
+		Mode = EBindMode::Runtime;
 	}
 
-	UE_LOGFMT(LogPlainPropsEngine, Display, "Binding types to PlainProps schemas...");
-	BindAllTypes();
+	SchemaBindAllTypes(Mode);
+	CustomBindEngineTypes(Mode);
 
 	// Parse roundtrip options
 	ERoundtrip Options = ERoundtrip::PP | ERoundtrip::UPS | ERoundtrip::TPS;
diff --git a/Engine/Plugins/Experimental/PlainPropsEngine/Source/Private/PlainPropsEngineBindings.cpp b/Engine/Plugins/Experimental/PlainPropsEngine/Source/Private/PlainPropsEngineBindings.cpp
--- a/Engine/Plugins/Experimental/PlainPropsEngine/Source/Private/PlainPropsEngineBindings.cpp
+++ b/Engine/Plugins/Experimental/PlainPropsEngine/Source/Private/PlainPropsEngineBindings.cpp
@@ -1,8 +1,16 @@
 // Copyright Epic Games, Inc. All Rights Reserved.
 
-#include "PlainPropsEngineModule.h"
+#include "PlainPropsEngineBindings.h"
+#include "Engine/Texture2D.h"
+#include "Logging/StructuredLog.h"
+
+DEFINE_LOG_CATEGORY(LogPlainPropsEngine);
 
 namespace PlainProps::UE
 {
 
-} // namespace PlainProps::UE
+void CustomBindEngineTypes(EBindMode Mode)
+{
+}
+
+} // namespace PlainProps::UE
\ No newline at end of file
diff --git a/Engine/Plugins/Experimental/PlainPropsEngine/Source/Private/PlainPropsEngineBindings.h b/Engine/Plugins/Experimental/PlainPropsEngine/Source/Private/PlainPropsEngineBindings.h
new file mode 100644
--- /dev/null
+++ b/Engine/Plugins/Experimental/PlainPropsEngine/Source/Private/PlainPropsEngineBindings.h
@@ -0,0 +1,14 @@
+// Copyright Epic Games, Inc. All Rights Reserved.
+
+#pragma once
+#include "Logging/StructuredLog.h"
+
+DECLARE_LOG_CATEGORY_EXTERN(LogPlainPropsEngine, Log, All);
+
+namespace PlainProps::UE
+{
+	enum class EBindMode : uint8;
+
+	void CustomBindEngineTypes(EBindMode Mode);
+
+} // namespace PlainProps::UE
\ No newline at end of file
diff --git a/Engine/Plugins/Experimental/PlainPropsEngine/Source/Private/PlainPropsEngineModule.cpp b/Engine/Plugins/Experimental/PlainPropsEngine/Source/Private/PlainPropsEngineModule.cpp
--- a/Engine/Plugins/Experimental/PlainPropsEngine/Source/Private/PlainPropsEngineModule.cpp
+++ b/Engine/Plugins/Experimental/PlainPropsEngine/Source/Private/PlainPropsEngineModule.cpp
@@ -3,8 +3,6 @@
 #include "PlainPropsEngineModule.h"
 #include "Modules/ModuleManager.h"
 
-DEFINE_LOG_CATEGORY(LogPlainPropsEngine);
-
 IMPLEMENT_MODULE(FPlainPropsEngineModule, PlainPropsEngine);
 
 void FPlainPropsEngineModule::StartupModule() {}
diff --git a/Engine/Plugins/Experimental/PlainPropsEngine/Source/Public/PlainPropsEngineModule.h b/Engine/Plugins/Experimental/PlainPropsEngine/Source/Public/PlainPropsEngineModule.h
--- a/Engine/Plugins/Experimental/PlainPropsEngine/Source/Public/PlainPropsEngineModule.h
+++ b/Engine/Plugins/Experimental/PlainPropsEngine/Source/Public/PlainPropsEngineModule.h
@@ -3,9 +3,6 @@
 #pragma once
 
 #include "Modules/ModuleInterface.h"
-#include "Logging/StructuredLog.h"
-
-DECLARE_LOG_CATEGORY_EXTERN(LogPlainPropsEngine, Log, All);
 
 class FPlainPropsEngineModule final : public IModuleInterface
 {
diff --git a/Engine/Plugins/Experimental/PlainPropsUObject/Source/Private/PlainPropsUObjectBindings.cpp b/Engine/Plugins/Experimental/PlainPropsUObject/Source/Private/PlainPropsUObjectBindings.cpp
--- a/Engine/Plugins/Experimental/PlainPropsUObject/Source/Private/PlainPropsUObjectBindings.cpp
+++ b/Engine/Plugins/Experimental/PlainPropsUObject/Source/Private/PlainPropsUObjectBindings.cpp
@@ -1919,14 +1919,20 @@ bool DiffMap(const BindingType& Binding, const ScriptType& A, const ScriptType&
 				// Todo: Add specific diff entry for missing key.
 				return true;
 			}
+
+			const void* PairA = KeyA;
+			const void* PairB = B.GetData(IdxB, Binding.Layout);
 			FBindId ValueId = Binding.PairRange.SaveId;
 			// Diff the whole Pair to generate the full nested diff path for diffing values
-			if (DiffStructs(/* GetValue */(KeyA), /* GetValue */(B.GetData(IdxB, Binding.Layout)), ValueId, Ctx))
+			if (DiffStructs(PairA, PairB, ValueId, Ctx))
 			{
-				check(!Binding.HasItem(B, A.GetData(IdxA, Binding.Layout)));
+				check(!Binding.HasItem(B, PairA));
 				return true;
 			}
-			check(Binding.HasItem(B, A.GetData(IdxA, Binding.Layout)));
+
+			// Temp comment: Fails for Lyra -loadmaps since FInstancedStruct is unbound,
+			//				 so all instanceed structs are considered equal
+			check(Binding.HasItem(B, PairA));
 		}
 	}
 
@@ -3505,13 +3511,20 @@ static void BindStruct(FBindId Id, const UStruct* Struct)
 	}
 }
 
-void DontBind(EPropertyFlags Skip)
+EPropertyFlags GetPropertiesToSkip(EBindMode Mode)
 {
-	GDontBind = Skip;
+	switch (Mode)
+	{
+	case EBindMode::Source:		return CPF_Transient | CPF_Deprecated | CPF_SkipSerialization;
+	case EBindMode::Runtime:	return CPF_Transient | CPF_Deprecated | CPF_SkipSerialization | CPF_EditorOnly;
+	default:					return CPF_None;
+	}
 }
 
-void BindAllTypes()
+void SchemaBindAllTypes(EBindMode Mode)
 {
+	GDontBind = GetPropertiesToSkip(Mode);
+	
 	// Declare all UEnums
 	for (TObjectIterator<UEnum> It; It; ++It)
 	{
@@ -5209,4 +5222,4 @@ void FPlainPropsUObjectModule::ShutdownModule()
 {
 	BindingScopePtr.Reset();
 	DebugIdScopePtr.Reset();
-}
\ No newline at end of file
+}
diff --git a/Engine/Plugins/Experimental/PlainPropsUObject/Source/Public/PlainPropsRoundtripTest.h b/Engine/Plugins/Experimental/PlainPropsUObject/Source/Public/PlainPropsRoundtripTest.h
--- a/Engine/Plugins/Experimental/PlainPropsUObject/Source/Public/PlainPropsRoundtripTest.h
+++ b/Engine/Plugins/Experimental/PlainPropsUObject/Source/Public/PlainPropsRoundtripTest.h
@@ -2,13 +2,16 @@
 
 #pragma once
 
-#include "CoreTypes.h"
+#include "Containers/ArrayView.h"
 #include "Misc/EnumClassFlags.h"
-#include "UObject/ObjectMacros.h"
+
+class UObject;
 
 namespace PlainProps::UE
 {
 
+enum class EBindMode : uint8 { All, Source, Runtime };
+
 enum class ERoundtrip : uint8
 {
 	None 		= 0,
@@ -20,9 +23,7 @@ enum class ERoundtrip : uint8
 };
 ENUM_CLASS_FLAGS(ERoundtrip);
 
-
-PLAINPROPSUOBJECT_API void DontBind(EPropertyFlags Skip);
-PLAINPROPSUOBJECT_API void BindAllTypes();
+PLAINPROPSUOBJECT_API void SchemaBindAllTypes(EBindMode Mode);
 PLAINPROPSUOBJECT_API int32 RoundtripViaBatch(TConstArrayView<UObject*> Objects, ERoundtrip Options);
 PLAINPROPSUOBJECT_API int32 RoundtripViaPackages(TConstArrayView<UObject*> Objects, ERoundtrip Options);
 
