[Core] Add support to compile switch between using sparse or compact sets as the default set container

#tests Using the editor, running PIE and game client, making builds

diff --git a/Engine/Plugins/Experimental/MeshModelingToolsetExp/Source/MeshModelingToolsExp/Private/ExtractSplineTool.cpp b/Engine/Plugins/Experimental/MeshModelingToolsetExp/Source/MeshModelingToolsExp/Private/ExtractSplineTool.cpp
--- a/Engine/Plugins/Experimental/MeshModelingToolsetExp/Source/MeshModelingToolsExp/Private/ExtractSplineTool.cpp
+++ b/Engine/Plugins/Experimental/MeshModelingToolsetExp/Source/MeshModelingToolsExp/Private/ExtractSplineTool.cpp
@@ -26,6 +26,7 @@
 
 #include "CuttingOps/PlaneCutOp.h"
 
+#include "Containers/SparseSet.h"
 #include "Misc/MessageDialog.h"
 
 #include "ModelingToolTargetUtil.h"
@@ -467,7 +468,7 @@ void UExtractSplineTool::GatherSplineDataFromPolygroupSelection()
 	// Basic algorithim:
 	// Starting with an edge from the set, sequentially pull edges out of the set which connect to the last selected edge. 
 	// Once an ordered list of boundary edges is generated, generate a list of spline points from these polygroup edges which connect into a continous polyline.
-	auto ProcessSelectedEdges = [this](TSet<int32>& ToProcessEdges)
+	auto ProcessSelectedEdges = [this](TSparseSet<int32>& ToProcessEdges)
 	{
 		if (ToProcessEdges.IsEmpty())
 		{
@@ -580,8 +581,9 @@ void UExtractSplineTool::GatherSplineDataFromPolygroupSelection()
 		return;
 	}
 
-	TSet<int32> MultipleCountedEdges;
-	TSet<int32> ToProcessEdges;
+	// Use sparse set so we can remove multiple and compact once
+	TSparseSet<int32> MultipleCountedEdges;
+	TSparseSet<int32> ToProcessEdges;
 	if (!Selection.SelectedGroupIDs.IsEmpty())
 	{
 		// Use this to develop equivalence relations between groups.
@@ -638,7 +640,12 @@ void UExtractSplineTool::GatherSplineDataFromPolygroupSelection()
 	}
 	else
 	{
-		ToProcessEdges = Selection.SelectedEdgeIDs;		
+		// Manual copy since SelectedEdgeIDs may be a different set type based on compile switches
+		ToProcessEdges.Empty(Selection.SelectedEdgeIDs.Num());
+		for (const int32 Id : Selection.SelectedEdgeIDs)
+		{
+			ToProcessEdges.Add(Id);
+		}
 
 		ProcessSelectedEdges(ToProcessEdges);
 	}
diff --git a/Engine/Plugins/Experimental/PlainProps/Source/Public/PlainPropsUeCoreBindings.h b/Engine/Plugins/Experimental/PlainProps/Source/Public/PlainPropsUeCoreBindings.h
--- a/Engine/Plugins/Experimental/PlainProps/Source/Public/PlainPropsUeCoreBindings.h
+++ b/Engine/Plugins/Experimental/PlainProps/Source/Public/PlainPropsUeCoreBindings.h
@@ -523,19 +523,26 @@ struct TSetBinding : IItemRangeBinding
 		}
 	}
 
+#if UE_USE_COMPACT_SET_AS_DEFAULT
 	virtual void ReadItems(FSaveRangeContext& Ctx) const override
 	{
-		static_assert(offsetof(TSetElement<T>, Value) == 0);
-		const TSparseArray<TSetElement<T>>& Elems = Ctx.Request.GetRange<TSparseArray<TSetElement<T>>>();
+		const SetType& Set = Ctx.Request.GetRange<SetType>();
+		Ctx.Items.SetAll(Set.GetData(), Set.Num());
+	}
+#else
+	virtual void ReadItems(FSaveRangeContext& Ctx) const override
+	{
+		static_assert(offsetof(TSparseSetElement<T>, Value) == 0);
+		const TSparseArray<TSparseSetElement<T>>& Elems = Ctx.Request.GetRange<TSparseArray<TSparseSetElement<T>>>();
 
 		if (Elems.IsEmpty())
 		{
-			Ctx.Items.SetAll(nullptr, 0, sizeof(TSetElement<T>));
+			Ctx.Items.SetAll(nullptr, 0, sizeof(TSparseSetElement<T>));
 		}
 		else if (FExistingItemSlice LastRead = Ctx.Items.Slice)
 		{
 			// Continue partial response
-			const TSetElement<T>* NextElem = static_cast<const TSetElement<T>*>(LastRead.Data) + LastRead.Num + /* skip known invalid */ 1;
+			const TSparseSetElement<T>* NextElem = static_cast<const TSparseSetElement<T>*>(LastRead.Data) + LastRead.Num + /* skip known invalid */ 1;
 			Ctx.Items.Slice = GetContiguousSlice(Elems.PointerToIndex(NextElem), Elems);
 		}
 		else if (Elems.IsCompact())
@@ -546,18 +553,19 @@ struct TSetBinding : IItemRangeBinding
 		{
 			// Start partial response
 			Ctx.Items.NumTotal = Elems.Num();
-			Ctx.Items.Stride = sizeof(TSetElement<T>);
+			Ctx.Items.Stride = sizeof(TSparseSetElement<T>);
 			Ctx.Items.Slice = GetContiguousSlice(0, Elems);
 		}
 	}
 
-	static FExistingItemSlice GetContiguousSlice(int32 Idx, const TSparseArray<TSetElement<T>>& Elems)
+	static FExistingItemSlice GetContiguousSlice(int32 Idx, const TSparseArray<TSparseSetElement<T>>& Elems)
 	{
 		int32 Num = 1;
 		for (;!Elems.IsValidIndex(Idx); ++Idx) {}
 		for (; Elems.IsValidIndex(Idx + Num); ++Num) {}
 		return { &Elems[Idx], static_cast<uint64>(Num) };
 	}
+#endif
 };
 
 //////////////////////////////////////////////////////////////////////////
diff --git a/Engine/Plugins/Experimental/PlainPropsUObject/Source/Private/PlainPropsUObjectBindings.cpp b/Engine/Plugins/Experimental/PlainPropsUObject/Source/Private/PlainPropsUObjectBindings.cpp
--- a/Engine/Plugins/Experimental/PlainPropsUObject/Source/Private/PlainPropsUObjectBindings.cpp
+++ b/Engine/Plugins/Experimental/PlainPropsUObject/Source/Private/PlainPropsUObjectBindings.cpp
@@ -1663,31 +1663,10 @@ using TPropertyRangeSerializer = typename TSelectRangeSerializer<Kind>::Type;
 
 ////////////////////////////////////////////////////////////////////////////////////////////////
 
+#if !UE_USE_COMPACT_SET_AS_DEFAULT
 inline FScriptSparseArray& AsSparseArray(FScriptSet& In) { return reinterpret_cast<FScriptSparseArray&>(In); }
 inline FScriptSparseArray& AsSparseArray(FScriptMap& In) { return reinterpret_cast<FScriptSparseArray&>(In); }
 
-template<typename ScriptSet>
-inline bool IsCompact(const ScriptSet& Set)
-{
-	return Set.NumUnchecked() == Set.GetMaxIndex();
-}
-
-// There's no TScriptSparseArray::SetNumUninitialized() (yet), 
-// reserve using Empty() and add items one by one instead
-template<class ScriptType, class LayoutType>
-uint8* SetNumUninitialized(ScriptType& Dst, const LayoutType& Layout, uint64 Num)
-{
-	check(Dst.IsEmpty());
-	Dst.Empty(static_cast<int32>(Num), Layout);
-	for (uint64 Idx = 0; Idx < Num; ++Idx)
-	{
-		Dst.AddUninitialized(Layout);
-	}
-	check(IsCompact(Dst));
-
-	return static_cast<uint8*>(Dst.GetData(0, Layout));
-}
-
 // @pre Elems.Num() > 0
 inline FExistingItemSlice GetContiguousSlice(int32 Idx, const FScriptSparseArray& Elems, const uint8* Data, SIZE_T Stride)
 {
@@ -1729,6 +1708,29 @@ inline void ReadSparseItems(FExistingItems& Dst, const FScriptSparseArray& Src,
 		Dst.Slice = GetContiguousSlice(0, Src, Data, Layout.Size);
 	}
 }
+#endif
+
+template<typename ScriptSet>
+inline bool IsCompact(const ScriptSet& Set)
+{
+	return Set.NumUnchecked() == Set.GetMaxIndex();
+}
+
+// There's no TScriptSparseArray::SetNumUninitialized() (yet), 
+// reserve using Empty() and add items one by one instead
+template<class ScriptType, class LayoutType>
+uint8* SetNumUninitialized(ScriptType& Dst, const LayoutType& Layout, uint64 Num)
+{
+	check(Dst.IsEmpty());
+	Dst.Empty(static_cast<int32>(Num), Layout);
+	for (uint64 Idx = 0; Idx < Num; ++Idx)
+	{
+		Dst.AddUninitialized(Layout);
+	}
+	check(IsCompact(Dst));
+
+	return static_cast<uint8*>(Dst.GetData(0, Layout));
+}
 
 ////////////////////////////////////////////////////////////////////////////////////////////////
 
@@ -1973,7 +1975,11 @@ struct TSetPropertyBinding : IItemRangeBinding, ICustomBinding
 		return Declare({Id, NoId, 0, EMemberPresence::AllowSparse, Members, {RangeSpec, RangeSpec, RangeSpec}});
 	}
 
-	inline SIZE_T GetStride() const requires (bLeaves)	{ return sizeof(TSetElement<LeafType>); }
+#if UE_USE_COMPACT_SET_AS_DEFAULT
+	inline SIZE_T GetStride() const requires (bLeaves)	{ return sizeof(LeafType); }
+#else
+	inline SIZE_T GetStride() const requires (bLeaves)	{ return sizeof(TSparseSetElement<LeafType>); }
+#endif
 	inline SIZE_T GetStride() const						{ return Layout.Size; }
 
 	inline int32 FindIndex(const FScriptSet& Set, const void* Elem) const
@@ -1986,7 +1992,11 @@ struct TSetPropertyBinding : IItemRangeBinding, ICustomBinding
 		if (int32 Idx = FindIndex(Set, Elem); Idx != INDEX_NONE)
 		{
 			DestroyElem(Set, Idx);
+#if UE_USE_COMPACT_SET_AS_DEFAULT
+			Set.RemoveAt(Idx, Layout, MakeHashFn(Inner));
+#else
 			Set.RemoveAt(Idx, Layout);
+#endif
 		}
 	}
 	
@@ -2000,11 +2010,6 @@ struct TSetPropertyBinding : IItemRangeBinding, ICustomBinding
 		return HasItem(Set, Elem);
 	}
 
-	inline void Rehash(FScriptSet& Set) const
-	{
-		Set.Rehash(Layout, MakeHashFn(Inner));
-	}
-
 	inline void DestroyElem(FScriptSet& Set) const requires (bLeaves) {}
 	inline void DestroyElem(FScriptSet& Set, int32 Idx) const
 	{
@@ -2035,7 +2040,12 @@ struct TSetPropertyBinding : IItemRangeBinding, ICustomBinding
 	
 	virtual void ReadItems(FSaveRangeContext& Ctx) const override
 	{
+#if UE_USE_COMPACT_SET_AS_DEFAULT
+		const FScriptSet& Set = Ctx.Request.GetRange<FScriptSet>();
+		Ctx.Items.SetAll(Set.GetData(0, Layout), Set.Num(), Layout.Size);
+#else
 		ReadSparseItems(/* out */ Ctx.Items, Ctx.Request.GetRange<FScriptSparseArray>(), Layout.SparseArrayLayout);
+#endif
 	}
 
 	virtual void MakeItems(FLoadRangeContext& Ctx) const override
@@ -2058,7 +2068,7 @@ struct TSetPropertyBinding : IItemRangeBinding, ICustomBinding
 		else
 		{
 			check(Ctx.Request.IsFinalCall());
-			Rehash(Set);
+			Set.Rehash(Layout, MakeHashFn(Inner));
 		}
 	}
 
@@ -2111,38 +2121,42 @@ struct TSetPropertyBinding : IItemRangeBinding, ICustomBinding
 			}
 		}
 
-		Rehash(Dst);
+		Dst.Rehash(Layout, MakeHashFn(Inner));
 	}
 
 	// Load leaves into non-empty set
 	inline void InsertNonEmpty(FScriptSet& Dst, FRangeLoadView Src) const requires(bLeaves)
 	{	
+		auto HashFn = MakeHashFn(Inner);
+
 		for (LeafType Item : CastAs<LeafType>(Src.AsLeaves()))
 		{
 			if (!HasItem(Dst, &Item))
 			{
 				 void* Elem = Dst.GetData(Dst.AddUninitialized(Layout), Layout);
+				 Dst.CommitLastUninitialized(Layout, HashFn);
 				 *static_cast<LeafType*>(Elem) = Item;
 			}
 		}
 
-		Rehash(Dst);
+		Dst.CommitAllUninitialized(Layout, HashFn);
 	}
 
-	inline void* AddItem(FScriptSparseArray& Dst, int32& OutIdx) const
+	inline void* AddItem(FScriptSet& Dst, int32& OutIdx) const
 	{
-		OutIdx = Dst.AddUninitialized(Layout.SparseArrayLayout);
-		void* Out = Dst.GetData(OutIdx, Layout.SparseArrayLayout);
+		OutIdx = Dst.AddUninitialized(Layout);
+		void* Out = Dst.GetData(OutIdx, Layout);
 		Inner.InitItem(Out);
 		return Out;
 	}
 
 	// Load structs or ranges into non-empty set
-	inline void InsertNonEmpty(FScriptSet& DstSet, FRangeLoadView Src) const
+	inline void InsertNonEmpty(FScriptSet& Dst, FRangeLoadView Src) const
 	{
+		auto HashFn = MakeHashFn(Inner);
+
 		// Written to avoid FProperty::CopyCompleteValue_InContainer dependency
 		// Items are loaded directly into sparse array and then removed if a duplicate existed
-		FScriptSparseArray& Dst = AsSparseArray(DstSet);
 		const int32 OldNum = Dst.NumUnchecked();
 		int32 TmpIdx;
 		void* Tmp = AddItem(Dst, /* out */ TmpIdx);
@@ -2152,7 +2166,11 @@ struct TSetPropertyBinding : IItemRangeBinding, ICustomBinding
 			for (FRangeLoadView Item : Src.AsRanges())
 			{
 				LoadRange(Tmp, Item, InnerBindings);
-				Tmp = HasItem(DstSet, Tmp) ? Tmp : AddItem(Dst, /* out */ TmpIdx);
+				if (!HasItem(Dst, Tmp))
+				{
+					Dst.CommitLastUninitialized(Layout, HashFn);
+					Tmp = AddItem(Dst, /* out */ TmpIdx);
+				}
 			}
 		}
 		else
@@ -2160,16 +2178,20 @@ struct TSetPropertyBinding : IItemRangeBinding, ICustomBinding
 			for (FStructLoadView Item : Src.AsStructs())
 			{
 				LoadStruct(Tmp, Item);
-				Tmp = HasItem(DstSet, Tmp) ? Tmp : AddItem(Dst, /* out */ TmpIdx);
+				if (!HasItem(Dst, Tmp))
+				{
+					Dst.CommitLastUninitialized(Layout, HashFn);
+					Tmp = AddItem(Dst, /* out */ TmpIdx);
+				}
 			}
 		}
 
 		Inner.DestroyItem(Tmp);
-		Dst.RemoveAtUninitialized(Layout.SparseArrayLayout, TmpIdx, 1);
+		Dst.RemoveAtUninitialized(Layout, TmpIdx);
 
 		if (Dst.NumUnchecked() != OldNum)
 		{
-			Rehash(DstSet);
+			Dst.CommitAllUninitialized(Layout, HashFn);
 		}
 	}
 
@@ -2213,7 +2235,7 @@ struct TSetPropertyBinding : IItemRangeBinding, ICustomBinding
 	{
 		checkSlow(Idx < Dst.GetMaxIndex());
 		checkSlow(!Dst.IsValidIndex(Idx));
-		void* Elem = AsSparseArray(Dst).GetData(Idx, Layout.SparseArrayLayout);
+		void* Elem = Dst.GetData(Idx, Layout);
 		Inner.DestroyItem(Elem);
 	}
 };
@@ -2335,7 +2357,12 @@ struct TMapPropertyItemBinding : IItemRangeBinding
 
 	virtual void ReadItems(FSaveRangeContext& Ctx) const override
 	{
+#if UE_USE_COMPACT_SET_AS_DEFAULT
+		const FScriptMap& Set = Ctx.Request.GetRange<FScriptMap>();
+		Ctx.Items.SetAll(Set.GetData(0, Layout), Set.Num(), Layout.SetLayout.Size);
+#else
 		ReadSparseItems(/* out */ Ctx.Items, Ctx.Request.GetRange<FScriptSparseArray>(), Layout.SetLayout.SparseArrayLayout);
+#endif
 	}
 
 	virtual void MakeItems(FLoadRangeContext& Ctx) const override
@@ -2431,12 +2458,12 @@ struct TMapPropertyCustomBinding : TMapPropertyItemBinding<FScriptMap, KeyKind,
 		return false;
 	}
 
-	inline uint8* AddPair(FScriptSparseArray& Dst, int32& OutIdx, int32 ValueOffset) const
+	inline uint8* AddPair(FScriptMap& Dst, int32& OutIdx) const
 	{
-		OutIdx = Dst.AddUninitialized(Layout.SetLayout.SparseArrayLayout);
-		uint8* Out = static_cast<uint8*>(Dst.GetData(OutIdx, Layout.SetLayout.SparseArrayLayout));
+		OutIdx = Dst.AddUninitialized(Layout);
+		uint8* Out = static_cast<uint8*>(Dst.GetData(OutIdx, Layout));
 		InnerKey.InitItem(Out);
-		InnerValue.InitItem(Out + ValueOffset);
+		InnerValue.InitItem(Out + Layout.ValueOffset);
 		return Out;
 	}
 
@@ -2455,7 +2482,11 @@ struct TMapPropertyCustomBinding : TMapPropertyItemBinding<FScriptMap, KeyKind,
 		if (int32 Idx = FindKey(Map, Key); Idx != INDEX_NONE)
 		{
 			DestroyPair(Map, Idx);
+#if UE_USE_COMPACT_SET_AS_DEFAULT
+			Map.RemoveAt(Idx, Layout, MakeHashFn(InnerKey));
+#else
 			Map.RemoveAt(Idx, Layout);
+#endif
 		}
 	}
 
@@ -2499,14 +2530,15 @@ struct TMapPropertyCustomBinding : TMapPropertyItemBinding<FScriptMap, KeyKind,
 	// Load structs or ranges into non-empty map
 	inline void InsertNonEmpty(FScriptMap& Dst, FRangeLoadView Src) const
 	{
+		auto HashFn = MakeHashFn(InnerKey);
+
 		// Written to avoid FProperty::CopyCompleteValue_InContainer dependency
 		// Items are loaded directly into sparse array and then removed if a duplicate existed
 
-		FScriptSparseArray& DstArray = AsSparseArray(Dst);
-		const int32 OldNum = DstArray.NumUnchecked();
+		const int32 OldNum = Dst.NumUnchecked();
 		const int32 ValueOffset = Layout.ValueOffset;
 		int32 TmpIdx;
-		uint8* Tmp = AddPair(DstArray, /* out*/ TmpIdx, ValueOffset);
+		uint8* Tmp = AddPair(Dst, /* out*/ TmpIdx);
 
 		FOptionalSchemaId InnerLoadIds[2];
 		FStructRangeLoadView Structs = Src.AsStructs();
@@ -2520,24 +2552,25 @@ struct TMapPropertyCustomBinding : TMapPropertyItemBinding<FScriptMap, KeyKind,
 			if (int32 Idx = FindKey(Dst, Tmp); Idx != INDEX_NONE)
 			{
 				// Load value into existing pair
-				uint8* Pair = static_cast<uint8*>(DstArray.GetData(Idx, Layout.SetLayout.SparseArrayLayout));
+				uint8* Pair = static_cast<uint8*>(Dst.GetData(Idx, Layout));
 				ValueRange.LoadItem(Pair + ValueOffset,	/* in-out */ Struct.Values, /* in-out */ Bits, InnerLoadIds[1], Struct.Schema.Batch);
 			}
 			else
 			{
 				// Load value into tmp pair and add new temporary
 				ValueRange.LoadItem(Tmp + ValueOffset,	/* in-out */ Struct.Values, /* in-out */ Bits, InnerLoadIds[1], Struct.Schema.Batch);
-				Tmp = AddPair(DstArray, /* out*/ TmpIdx, ValueOffset);
+				Dst.CommitLastUninitialized(Layout, HashFn);
+				Tmp = AddPair(Dst, /* out*/ TmpIdx);
 			}
 		}
 
 		InnerKey.DestroyItem(Tmp);
 		InnerValue.DestroyItem(Tmp + ValueOffset);
-		DstArray.RemoveAtUninitialized(Layout.SetLayout.SparseArrayLayout, TmpIdx, 1);
+		Dst.RemoveAtUninitialized(Layout, TmpIdx);
 
-		if (DstArray.NumUnchecked() != OldNum)
+		if (Dst.NumUnchecked() != OldNum)
 		{
-			Rehash(Dst);
+			Dst.CommitAllUninitialized(Layout, HashFn);
 		}
 	}
 
@@ -3503,7 +3536,7 @@ static void BindStruct(FBindId Id, const UStruct* Struct)
 	BindMembers(/* out */ Binder, Struct);
 
 	GUE.Schemas.BindStruct(Id, Binder.GetMembers(), Binder.Declare());
-		
+	
 	// Don't bind CDO defaults, object defaults are passed in from top and objects aren't owned by containers 
 	if (Struct->HasAnyCastFlags(CASTCLASS_UScriptStruct))
 	{
@@ -4857,18 +4890,18 @@ int32 RoundtripViaBatch(TConstArrayView<UObject*> Objects, ERoundtrip Options)
 	}
 	else
 	{
-		UE_LOGFMT(LogPlainPropsUObject, Display, "Gathering all non-empty UObjects...");
-		static constexpr EObjectFlags SkipFlags = RF_ClassDefaultObject | RF_MirroredGarbage | RF_InheritableComponentTemplate;
-		for (TObjectIterator<UObject> It(SkipFlags); It; ++It)
-		{
-			UObject* Object = *It;
+	UE_LOGFMT(LogPlainPropsUObject, Display, "Gathering all non-empty UObjects...");
+	static constexpr EObjectFlags SkipFlags = RF_ClassDefaultObject | RF_MirroredGarbage | RF_InheritableComponentTemplate;
+	for (TObjectIterator<UObject> It(SkipFlags); It; ++It)
+	{
+		UObject* Object = *It;
 			UClass* Class = Object->GetClass();
 			if (ShouldBind(Class) && IncludeClass(Class))
-			{
-				Instances.Emplace(Object->GetPathName(), Object);
-			}
+		{
+			Instances.Emplace(Object->GetPathName(), Object);
 		}
 	}
+	}
 
 	UE_LOGFMT(LogPlainPropsUObject, Display, "Sorting {Num} UObjects ...", Instances.Num());
 	Algo::Sort(Instances, [](const FInstance& A, const FInstance& B)
@@ -5107,20 +5140,20 @@ int32 RoundtripViaPackages(TConstArrayView<UObject*> Objects, ERoundtrip Options
 		{
 			UClass* Class = Object->GetClass();
 			if (Class->HasAnyClassFlags(CLASS_CompiledFromBlueprint))
-			{
+		{
 				bExclude = true;
 				return false;
-			}
+		}
 			if (!IncludeClass(Class))
 			{
 				UE_LOGFMT(LogPlainPropsUObject, Log, " Excluding PlainProps for {Package}", Object->GetPackage()->GetName());
 				bExclude = true;
 				return false;
-			}
+	}
 			return true;
 		});
 		if (bExclude)
-		{
+	{
 			continue;
 		}
 		Packages.Add(Package);
@@ -5137,7 +5170,7 @@ int32 RoundtripViaPackages(TConstArrayView<UObject*> Objects, ERoundtrip Options
 
 	UE_LOGFMT(LogPlainPropsUObject, Display, "Saving {SaveNum} packages...", Packages.Num());
 	for (UPackage* Package : Packages)
-	{
+		{
 		UObject* Asset = Package->FindAssetInPackage();
 		FSavePackageArgs SaveArgs;
 		SaveArgs.TopLevelFlags = RF_Standalone;
@@ -5222,4 +5255,4 @@ void FPlainPropsUObjectModule::ShutdownModule()
 {
 	BindingScopePtr.Reset();
 	DebugIdScopePtr.Reset();
-}
+}
\ No newline at end of file
diff --git a/Engine/Plugins/Runtime/AudioModulation/Source/AudioModulation/Private/AudioModulationDebugger.cpp b/Engine/Plugins/Runtime/AudioModulation/Source/AudioModulation/Private/AudioModulationDebugger.cpp
--- a/Engine/Plugins/Runtime/AudioModulation/Source/AudioModulation/Private/AudioModulationDebugger.cpp
+++ b/Engine/Plugins/Runtime/AudioModulation/Source/AudioModulation/Private/AudioModulationDebugger.cpp
@@ -565,7 +565,7 @@ namespace AudioModulation
 				// Bus mix is global if it only has one stage and its ID is the Control Bus's ID.
 				if (Proxy->Stages.Num() == 1)
 				{
-					uint32 BusId = Proxy->Stages.begin().Key();
+					uint32 BusId = Proxy->Stages.FindArbitraryElement()->Key;
 					DebugInfo.bIsGlobal = ModulationSystem->IsBusMixGlobal(BusId);
 				}
 			
diff --git a/Engine/Source/Editor/ContentBrowser/Private/SAssetView.cpp b/Engine/Source/Editor/ContentBrowser/Private/SAssetView.cpp
--- a/Engine/Source/Editor/ContentBrowser/Private/SAssetView.cpp
+++ b/Engine/Source/Editor/ContentBrowser/Private/SAssetView.cpp
@@ -625,7 +625,7 @@ private:
 
 bool FAssetViewItemCollection::RefreshLookup()
 {
-	uint32 HashSize = TSetAllocator<>::GetNumberOfHashBuckets(Items.Num());
+	uint32 HashSize = FDefaultSparseSetAllocator::GetNumberOfHashBuckets(Items.Num());
 	if (Lookup.GetHashSize() < HashSize)
 	{
 		Lookup.Clear(HashSize, Items.Max());
@@ -3484,7 +3484,7 @@ void FAssetViewItemCollection::RefreshItemsFromBackend(const FAssetViewContentSo
 	bItemsPendingRemove = false;
 
 	// Create or recyle FAssetViewItem for each FContentBrowserItemData. Build new hashtable concurrently at the same time
-	uint32 HashSize = TSetAllocator<>::GetNumberOfHashBuckets(Items.Num());
+	uint32 HashSize = FDefaultSparseSetAllocator::GetNumberOfHashBuckets(Items.Num());
 	std::atomic<bool> bAnyFolders(false);
 	std::atomic<bool> bAnyRecycled(false);
 	Lookup.Clear(HashSize, Items.Num());
diff --git a/Engine/Source/Editor/UnrealEd/Private/DataTableEditorUtils.cpp b/Engine/Source/Editor/UnrealEd/Private/DataTableEditorUtils.cpp
--- a/Engine/Source/Editor/UnrealEd/Private/DataTableEditorUtils.cpp
+++ b/Engine/Source/Editor/UnrealEd/Private/DataTableEditorUtils.cpp
@@ -465,15 +465,12 @@ bool FDataTableEditorUtils::RemoveRow(UDataTable* DataTable, FName Name)
 		BroadcastPreChange(DataTable, EDataTableChangeInfo::RowList);
 		DataTable->Modify();
 		uint8* RowData = nullptr;
-		const bool bRemoved = DataTable->GetNonConstRowMap().RemoveAndCopyValue(Name, RowData);
+		const bool bRemoved = DataTable->GetNonConstRowMap().RemoveAndCopyValueStable(Name, RowData);
 		if (bRemoved && RowData)
 		{
 			DataTable->RowStruct->DestroyStruct(RowData);
 			FMemory::Free(RowData);
 			bResult = true;
-
-			// Compact the map so that a subsequent add goes at the end of the table
-			DataTable->GetNonConstRowMap().CompactStable();
 		}
 		BroadcastPostChange(DataTable, EDataTableChangeInfo::RowList);
 	}
diff --git a/Engine/Source/Runtime/Core/Private/Misc/ConfigCacheIni.cpp b/Engine/Source/Runtime/Core/Private/Misc/ConfigCacheIni.cpp
--- a/Engine/Source/Runtime/Core/Private/Misc/ConfigCacheIni.cpp
+++ b/Engine/Source/Runtime/Core/Private/Misc/ConfigCacheIni.cpp
@@ -1766,13 +1766,11 @@ void FConfigFile::ProcessCommand(FConfigSection* Section, FStringView SectionNam
 			break;
 		case FConfigValue::EValueType::Remove:
 			// Remove if present.
-			Section->RemoveSingle(Key, Value);
-			Section->CompactStable();
+			Section->RemoveSingleStable(Key, Value);
 			break;
 		case FConfigValue::EValueType::Clear:
 			// Remove if present.
-			Section->Remove(Key);
-			Section->CompactStable();
+			Section->RemoveStable(Key);
 			
 			// clear any empty initialization, which will reset it to using code defaults if nothing else adds anything else
 			Section->EmptyInitializedKeys.Remove(Key);
@@ -1782,8 +1780,7 @@ void FConfigFile::ProcessCommand(FConfigSection* Section, FStringView SectionNam
 			Section->EmptyInitializedKeys.Add(Key);
 			
 			// also clear any existing entries
-			Section->Remove(Key);
-			Section->CompactStable();
+			Section->RemoveStable(Key);
 			break;
 		case FConfigValue::EValueType::ArrayOfStructKey:
 			// track a key to show uniqueness for arrays of structs
diff --git a/Engine/Source/Runtime/Core/Public/Containers/Array.h b/Engine/Source/Runtime/Core/Public/Containers/Array.h
--- a/Engine/Source/Runtime/Core/Public/Containers/Array.h
+++ b/Engine/Source/Runtime/Core/Public/Containers/Array.h
@@ -216,6 +216,18 @@ UE_NODEBUG [[nodiscard]] UE_FORCEINLINE_HINT TIndexedContainerIterator<Container
 		{
 		}
 
+		UE_NODEBUG [[nodiscard]] UE_FORCEINLINE_HINT ElementType* operator->() const
+		{
+			if constexpr (bReverse)
+			{
+				return Ptr - 1;
+			}
+			else
+			{
+				return Ptr;
+			}
+		}
+
 		UE_NODEBUG [[nodiscard]] UE_FORCEINLINE_HINT ElementType& operator*() const
 		{
 			if constexpr (bReverse)
@@ -265,6 +277,12 @@ UE_NODEBUG [[nodiscard]] UE_FORCEINLINE_HINT TIndexedContainerIterator<Container
 			return Ptr != Rhs.Ptr;
 		}
 
+		UE_NODEBUG [[nodiscard]] FORCEINLINE bool operator==(const TCheckedPointerIterator& Rhs) const
+		{
+			return !(*this != Rhs);
+		}
+
+
 	private:
 		ElementType*    Ptr;
 		const SizeType& CurrentNum;
diff --git a/Engine/Source/Runtime/Core/Public/Containers/CompactSet.h b/Engine/Source/Runtime/Core/Public/Containers/CompactSet.h
--- a/Engine/Source/Runtime/Core/Public/Containers/CompactSet.h
+++ b/Engine/Source/Runtime/Core/Public/Containers/CompactSet.h
@@ -111,6 +111,10 @@ private:
 
 		UE_FORCEINLINE_HINT TBaseIterator& operator++()
 		{
+#if DO_CHECK
+			const int32 SetNum = Set.Num();
+			checkf(SetNum >= InitialNum, TEXT("Sets/Maps should never have elements removed during iteration outside of Iterator.RemoveCurrent(). InitialNum %d CurrentNum %d"), InitialNum, SetNum);
+#endif
 			++Index;
 			return *this;
 		}
@@ -123,9 +127,6 @@ private:
 
 		[[nodiscard]] inline bool operator==(const TBaseIterator& Rhs) const
 		{
-#if DO_CHECK
-			checkf(Set.Num() == InitialNum, TEXT("Sets/Maps should never be mutated during range-for iteration. Use CreateIterator form to allow more flexibility. InitialNum %d CurrentNum %d"), InitialNum, Set.Num());
-#endif
 			checkSlow(&Set == &Rhs.Set);
 			return Index == Rhs.Index;
 		}
@@ -185,7 +186,7 @@ private:
 
 			// Note: Adding new elements is safe, however they will be guaranteed to be missed by the current iteration
 #if DO_CHECK
-			checkf(SetNum >= InitialNum, TEXT("Sets/Maps should never have elements removed during iteration outside of Iterator.RemoveCurrent(). InitialNum %d CurrentNum %d"), InitialNum, Set.Num());
+			checkf(SetNum >= InitialNum, TEXT("Sets/Maps should never have elements removed during iteration outside of Iterator.RemoveCurrent(). InitialNum %d CurrentNum %d"), InitialNum, SetNum);
 #endif
 
 			Index = NextIndex;
@@ -300,6 +301,11 @@ public:
 
 	[[nodiscard]] UE_FORCEINLINE_HINT TCompactSet() = default;
 
+	[[nodiscard]] explicit consteval TCompactSet(EConstEval)
+		: Super(ConstEval)
+	{
+	}
+
 	[[nodiscard]] UE_FORCEINLINE_HINT TCompactSet(const TCompactSet& Copy)
 	{
 		*this = Copy;
@@ -553,7 +559,7 @@ public:
 	 */
 	[[nodiscard]] UE_FORCEINLINE_HINT SIZE_T GetAllocatedSize() const
 	{
-		return this->Elements.GetAllocationSize();
+		return Super::GetAllocatedSize(GetSetLayout());
 	}
 
 	/** Tracks the container's memory use through an archive. */
@@ -606,10 +612,10 @@ public:
 	}
 
 	/**
-		* Checks if index is in array range.
-		*
-		* @param Index Index to check.
-		*/
+	 * Checks if index is in array range.
+	 *
+	 * @param Index Index to check.
+	 */
 	inline void RangeCheck(FSetElementId Id) const
 	{
 		CheckInvariants();
@@ -669,6 +675,7 @@ public:
 	 * Adds an element to the set.
 	 *
 	 * @see		Class documentation section on ByHash() functions
+	 * @param	KeyHash						A precomputed hash value, calculated in the same way as ElementType is hashed.
 	 * @param	InElement					Element to add to set
 	 * @param	bIsAlreadyInSetPtr	[out]	Optional pointer to bool that will be set depending on whether element is already in set
      * @return  A handle to the element stored in the set
@@ -685,15 +692,35 @@ public:
 	/**
 	 * Adds an element to the set.
 	 *
-	 * @param	Args						The argument(s) to be forwarded to the set element's constructor.
+	 * @param	Arg							The argument(s) to be forwarded to the set element's constructor.
 	 * @param	bIsAlreadyInSetPtr	[out]	Optional pointer to bool that will be set depending on whether element is already in set
 	 * @return	A handle to the element stored in the set.
 	 */
-	template<typename ArgsType = ElementType>
-	FSetElementId Emplace(ArgsType && Args, bool* bIsAlreadyInSetPtr = nullptr)
+	template<typename ArgType = ElementType>
+	UE_FORCEINLINE_HINT FSetElementId Emplace(ArgType && Arg, bool* bIsAlreadyInSetPtr = nullptr)
 	{
+		TPair<FSetElementId, bool> Result = Emplace(InPlace, Forward<ArgType>(Arg));
+
+		if (bIsAlreadyInSetPtr)
+		{
+			*bIsAlreadyInSetPtr = Result.Value;
+		}
+
+		return Result.Key;
+	}
+
+    /**
+     * Adds an element to the set by constructing the ElementType in-place with multiple args.
+     *
+     * @param   EInPlace   Tag to disambiguate in-place construction.
+     * @param   InArgs     Arguments forwarded to ElementType's constructor.
+     * @return  Pair of (element id, whether an equivalent element already existed).
+     */
+    template <typename... ArgTypes>
+    TPair<FSetElementId, bool> Emplace(EInPlace, ArgTypes&&... InArgs)
+    {
 		alignas(alignof(ElementType))char StackElement[sizeof(ElementType)];
-		ElementType& TempElement = *new (StackElement) ElementType(Forward<ArgsType>(Args));
+		ElementType& TempElement = *new (StackElement) ElementType(Forward<ArgTypes>(InArgs)...);
 		SizeType ExistingIndex(INDEX_NONE);
 
 		const KeyInitType Key = KeyFuncs::GetSetKey(TempElement);
@@ -704,11 +731,6 @@ public:
 			ExistingIndex = FindIndexByHash(KeyHash, Key);
 		}
 
-		if (bIsAlreadyInSetPtr)
-		{
-			*bIsAlreadyInSetPtr = ExistingIndex != INDEX_NONE;
-		}
-
 		if (ExistingIndex == INDEX_NONE)
 		{
 			ExistingIndex = this->NumElements;
@@ -721,22 +743,45 @@ public:
 			MoveByRelocate(ExistingElement, TempElement);
 		}
 
-		return FSetElementId::FromInteger(ExistingIndex);
+		return { FSetElementId::FromInteger(ExistingIndex), ExistingIndex != INDEX_NONE };
 	}
 	
 	/**
 	 * Adds an element to the set.
 	 *
 	 * @see		Class documentation section on ByHash() functions
-	 * @param	Args						The argument(s) to be forwarded to the set element's constructor.
+	 * @param	KeyHash						A precomputed hash value, calculated in the same way as ElementType is hashed.
+	 * @param	Arg							The argument(s) to be forwarded to the set element's constructor.
 	 * @param	bIsAlreadyInSetPtr	[out]	Optional pointer to bool that will be set depending on whether element is already in set
 	 * @return	A handle to the element stored in the set.
 	 */
-	template<typename ArgsType = ElementType>
-	FSetElementId EmplaceByHash(uint32 KeyHash, ArgsType && Args, bool* bIsAlreadyInSetPtr = nullptr)
+	template<typename ArgType = ElementType>
+	UE_FORCEINLINE_HINT FSetElementId EmplaceByHash(uint32 KeyHash, ArgType && Arg, bool* bIsAlreadyInSetPtr = nullptr)
 	{
+		TPair<FSetElementId, bool> Result = EmplaceByHash(InPlace, KeyHash, Forward<ArgType>(Arg));
+
+		if (bIsAlreadyInSetPtr)
+		{
+			*bIsAlreadyInSetPtr = Result.Value;
+		}
+
+		return Result.Key;
+	}
+
+    /**
+     * Adds an element to the set by constructing in-place with multiple args, using a precomputed hash.
+	 *
+     * @see     Class documentation section on ByHash() functions.
+	 * @param	KeyHash    A precomputed hash value, calculated in the same way as ElementType is hashed.
+     * @param   EInPlace   Tag to disambiguate in-place construction.
+     * @param   InArgs     Arguments forwarded to ElementType's constructor.
+     * @return  Pair of (element id, whether an equivalent element already existed).
+     */
+    template <typename... ArgTypes>
+    TPair<FSetElementId, bool> EmplaceByHash(EInPlace, uint32 KeyHash, ArgTypes&&... InArgs)
+    {
 		alignas(alignof(ElementType))char StackElement[sizeof(ElementType)];
-		ElementType& TempElement = *new (StackElement) ElementType(Forward<ArgsType>(Args));
+		ElementType& TempElement = *new (StackElement) ElementType(Forward<ArgTypes>(InArgs)...);
 		SizeType ExistingIndex(INDEX_NONE);
 
 		if constexpr (!KeyFuncs::bAllowDuplicateKeys)
@@ -744,11 +789,6 @@ public:
 			ExistingIndex = FindIndexByHash(KeyHash, KeyFuncs::GetSetKey(TempElement));
 		}
 
-		if (bIsAlreadyInSetPtr)
-		{
-			*bIsAlreadyInSetPtr = ExistingIndex != INDEX_NONE;
-		}
-
 		if (ExistingIndex == INDEX_NONE)
 		{
 			ExistingIndex = this->NumElements;
@@ -761,16 +801,16 @@ public:
 			MoveByRelocate(ExistingElement, TempElement);
 		}
 
-		return FSetElementId::FromInteger(ExistingIndex);
+		return { FSetElementId::FromInteger(ExistingIndex), ExistingIndex != INDEX_NONE };
 	}
 
 	/**
-		* Adds an element to the set if not already present and returns a reference to the added or existing element.
-		*
-		* @param	InElement					Element to add to set
-		* @param	bIsAlreadyInSetPtr	[out]	Optional pointer to bool that will be set depending on whether element is already in set
-		* @return	A reference to the element stored in the set.
-		*/
+	 * Adds an element to the set if not already present and returns a reference to the added or existing element.
+	 *
+	 * @param	InElement					Element to add to set
+	 * @param	bIsAlreadyInSetPtr	[out]	Optional pointer to bool that will be set depending on whether element is already in set
+	 * @return	A reference to the element stored in the set.
+	 */
 	UE_FORCEINLINE_HINT ElementType& FindOrAdd(const InElementType& InElement, bool* bIsAlreadyInSetPtr = nullptr)
 	{
 		return FindOrAddByHash(KeyFuncs::GetKeyHash(KeyFuncs::GetSetKey(InElement)), InElement, bIsAlreadyInSetPtr);
@@ -784,6 +824,7 @@ public:
 	 * Adds an element to the set if not already present and returns a reference to the added or existing element.
 	 *
 	 * @see		Class documentation section on ByHash() functions
+	 * @param	KeyHash						A precomputed hash value, calculated in the same way as ElementType is hashed.
 	 * @param	InElement					Element to add to set
 	 * @param	bIsAlreadyInSetPtr	[out]	Optional pointer to bool that will be set depending on whether element is already in set
 	 * @return  A reference to the element stored in the set
@@ -903,10 +944,10 @@ public:
 	}
 
 	/**
-		* Finds any element in the set and returns a pointer to it.
-		* Callers should not depend on particular patterns in the behaviour of this function.
-		* @return A pointer to an arbitrary element, or nullptr if the container is empty.
-		*/
+	 * Finds any element in the set and returns a pointer to it.
+	 * Callers should not depend on particular patterns in the behaviour of this function.
+	 * @return A pointer to an arbitrary element, or nullptr if the container is empty.
+	 */
 	[[nodiscard]] ElementType* FindArbitraryElement()
 	{
 		return (this->NumElements > 0) ? GetData() : nullptr;
@@ -1402,10 +1443,10 @@ private:
 	}
 
 	/**
-		* Finds an element with a pre-calculated hash and a key that can be compared to KeyType
-		* @see	Class documentation section on ByHash() functions
-		* @return The element id that matches the key and hash or an invalid element id
-		*/
+	 * Finds an element with a pre-calculated hash and a key that can be compared to KeyType
+	 * @see	Class documentation section on ByHash() functions
+	 * @return The element id that matches the key and hash or an invalid element id
+	 */
 	template<typename ComparableKey>
 	[[nodiscard]] SizeType FindIndexByHash(uint32 KeyHash, const ComparableKey& Key) const
 	{
diff --git a/Engine/Source/Runtime/Core/Public/Containers/CompactSetBase.h b/Engine/Source/Runtime/Core/Public/Containers/CompactSetBase.h
--- a/Engine/Source/Runtime/Core/Public/Containers/CompactSetBase.h
+++ b/Engine/Source/Runtime/Core/Public/Containers/CompactSetBase.h
@@ -6,89 +6,6 @@
 #include "Containers/CompactHashTable.h"
 #include "Templates/AlignmentTemplates.h"
 
-//
-// Compact set allocation definitions - It operates entirely with uint8 because it includes its lookup table inline with the overall allocation
-// Todo: move to ContainerAllocationPolicies.h
-//
-
-namespace UE::Core::Private::CompactSetAllocatorHelpers
-{
-	template <uint32 NumInlineElements, int32 ElementSize>
-	constexpr int32 CalculateRequiredBytes()
-	{
-		constexpr uint32 TypeSize = 1 + (NumInlineElements > 0xff) + (NumInlineElements > 0xffff) * 2;
-		constexpr uint32 HashSize = NumInlineElements < 8 ? 4 : FGenericPlatformMath::RoundUpToPowerOfTwo((NumInlineElements / 2) + 1);
-		return Align(NumInlineElements*ElementSize, 4) + 4 + (NumInlineElements + HashSize) * TypeSize;
-	}
-
-	/** Regular iterator */
-	template<typename InElementAllocator = FDefaultAllocator>
-	struct TCompactSetAllocator
-	{
-		/** CompactSets use raw uint8, so we need extra alignment information for inline allocators */
-		template <typename ElementType>
-		struct AllocatorAlignment
-		{
-			static constexpr size_t Value = alignof(typename InElementAllocator::template ForElementType<uint8>);
-		};
-
-		template<int32 ElementSize>
-		using ElementAllocator = InElementAllocator;
-	};
-
-	/** An inline allocator that allows sizing of the inline allocations for a set number of elements. */
-	template<uint32 NumInlineElements, typename SecondaryAllocator = TCompactSetAllocator<>>
-	struct TInlineCompactSetAllocator
-	{
-		template<int32 ElementSize>
-		using ElementAllocator = TInlineAllocator<UE::Core::Private::CompactSetAllocatorHelpers::CalculateRequiredBytes<NumInlineElements, ElementSize>(),typename SecondaryAllocator::template ElementAllocator<ElementSize>>;
-
-		/** CompactSets use raw uint8, so we need extra alignment information for inline allocators */
-		template <typename ElementType>
-		struct AllocatorAlignment
-		{
-			static constexpr size_t ElementAlignof = alignof(ElementType);
-			static constexpr size_t AllocatorAlignof = alignof(typename ElementAllocator<sizeof(ElementType)>::template ForElementType<uint8>);
-			static constexpr size_t Value = FGenericPlatformMath::Max(ElementAlignof, AllocatorAlignof);
-		};
-	};
-
-	/** An inline allocator that doesn't have any secondary storage. */
-	template <uint32 NumInlineElements>
-	struct TFixedCompactSetAllocator
-	{
-		template<int32 ElementSize>
-		using ElementAllocator = TFixedAllocator<UE::Core::Private::CompactSetAllocatorHelpers::CalculateRequiredBytes<NumInlineElements, ElementSize>()>;
-	
-		/** CompactSets use raw uint8, so we need extra alignment information for inline allocators */
-		template <typename ElementType>
-		struct AllocatorAlignment
-		{
-			static constexpr size_t ElementAlignof = alignof(ElementType);
-			static constexpr size_t AllocatorAlignof = alignof(typename ElementAllocator<sizeof(ElementType)>::template ForElementType<uint8>);
-			static constexpr size_t Value = FGenericPlatformMath::Max(ElementAlignof, AllocatorAlignof);
-		};
-	};
-}
-
-template<typename InElementAllocator>
-struct TAllocatorTraits<UE::Core::Private::CompactSetAllocatorHelpers::TCompactSetAllocator<InElementAllocator>> : TAllocatorTraitsBase<UE::Core::Private::CompactSetAllocatorHelpers::TCompactSetAllocator<InElementAllocator>>
-{
-	enum
-	{
-		SupportsFreezeMemoryImage = TAllocatorTraits<InElementAllocator>::SupportsFreezeMemoryImage,
-	};
-};
-
-class FDefaultCompactSetAllocator  : public UE::Core::Private::CompactSetAllocatorHelpers::TCompactSetAllocator<>  { public: typedef UE::Core::Private::CompactSetAllocatorHelpers::TCompactSetAllocator<>  Typedef; };
-
-// End Compact set allocation definitions
-
-// Todo: Move to ContainersFwd.h
-class FDefaultCompactSetAllocator;
-template<typename InElementType, typename KeyFuncs = DefaultKeyFuncs<InElementType>, typename Allocator = FDefaultCompactSetAllocator> class TCompactSet;
-// End Move to ContainersFwd.h
-
 /** Used to describe the data layout of the contents of the set */
 struct FCompactSetLayout
 {
@@ -109,22 +26,50 @@ public:
 		return MaxElements == INDEX_NONE;
 	}
 
+	/**
+	 * Returns true if the sets is empty and contains no elements. 
+	 *
+	 * @returns True if the set is empty.
+	 * @see Num
+	 */
 	[[nodiscard]] UE_FORCEINLINE_HINT bool IsEmpty() const
 	{
 		return NumElements == 0;
 	}
 
+	/** @return the number of elements. */
 	[[nodiscard]] UE_FORCEINLINE_HINT int32 Num() const
 	{
 		return NumElements;
 	}
 
+	/** @return The number of elements the set can hold before reallocation. */
+	[[nodiscard]] UE_FORCEINLINE_HINT int32 Max() const
+	{
+		return MaxElements;
+	}
+
 	/** Deprecated - unnecessary, keeping this here so TCompactSet can be swapped with TSet without changing code */
 	[[nodiscard]] UE_FORCEINLINE_HINT int32 GetMaxIndex() const
 	{
 		return NumElements;
 	}
 
+	/** 
+	 * Helper function to return the amount of memory allocated by this container 
+	 * Only returns the size of allocations made directly by the container, not the elements themselves.
+	 * @return number of bytes allocated by this container
+	 */
+	[[nodiscard]] UE_FORCEINLINE_HINT SIZE_T GetAllocatedSize(const FCompactSetLayout Layout) const
+	{
+		if (MaxElements == 0)
+		{
+			return 0;
+		}
+
+		return GetTotalMemoryRequiredInBytes(MaxElements, *GetHashTableMemory(Layout), Layout);
+	}
+
 protected:
 	using HashCountType = uint32;
 	static constexpr size_t HashCountSize = sizeof(HashCountType);
@@ -133,7 +78,14 @@ protected:
 	static_assert(std::is_same_v<SizeType, int32>, "TCompactSet currently only supports 32-bit allocators");
 	static_assert(sizeof(HashCountType) == UE::Core::CompactHashTable::GetMemoryAlignment(), "Hashtable alignment changed, need to make sure HashCountType is stored correctly");
 
-	[[nodiscard]] UE_FORCEINLINE_HINT TCompactSetBase() = default;	
+	[[nodiscard]] UE_FORCEINLINE_HINT TCompactSetBase() = default;
+
+	[[nodiscard]] explicit consteval TCompactSetBase(EConstEval)
+	: Elements(ConstEval)
+	, NumElements(0)
+	, MaxElements(0)
+	{
+	}
 
 	[[nodiscard]] explicit TCompactSetBase(FIntrusiveUnsetOptionalState Tag)
 	: NumElements(0)
diff --git a/Engine/Source/Runtime/Core/Public/Containers/ContainerAllocationPolicies.h b/Engine/Source/Runtime/Core/Public/Containers/ContainerAllocationPolicies.h
--- a/Engine/Source/Runtime/Core/Public/Containers/ContainerAllocationPolicies.h
+++ b/Engine/Source/Runtime/Core/Public/Containers/ContainerAllocationPolicies.h
@@ -8,6 +8,7 @@
 #include "HAL/UnrealMemory.h"
 #include "Math/NumericLimits.h"
 #include "Misc/AssertionMacros.h"
+#include "Templates/AlignmentTemplates.h"
 #include "Templates/IsPolymorphic.h"
 #include "Templates/MemoryOps.h"
 #include "Templates/TypeCompatibleBytes.h"
@@ -134,12 +135,6 @@ struct FArraySlackTrackingHeader
 static_assert(UE_CONTAINER_SLACK_GROWTH_FACTOR_DENOMINATOR > 0, "UE_CONTAINER_SLACK_GROWTH_FACTOR_DENOMINATOR must be greater than 0");
 static_assert(UE_CONTAINER_SLACK_GROWTH_FACTOR_DENOMINATOR > UE_CONTAINER_SLACK_GROWTH_FACTOR_NUMERATOR, "UE_CONTAINER_SLACK_GROWTH_FACTOR_DENOMINATOR must be greater than UE_CONTAINER_SLACK_GROWTH_FACTOR_NUMERATOR");
 
-
-class FDefaultBitArrayAllocator;
-
-template<int IndexSize> class TSizedDefaultAllocator;
-using FDefaultAllocator = TSizedDefaultAllocator<32>;
-
 template <typename SizeType>
 UE_FORCEINLINE_HINT SizeType DefaultCalculateSlackShrink(SizeType NewMax, SizeType CurrentMax, SIZE_T BytesPerElement, bool bAllowQuantize, uint32 Alignment = DEFAULT_ALIGNMENT)
 {
@@ -623,6 +618,7 @@ class TSizedHeapAllocator
 {
 public:
 	using SizeType = typename TBitsToSizeType<IndexSize>::Type;
+	using BaseMalloc = BaseMallocType;
 
 private:
 	using USizeType = std::make_unsigned_t<SizeType>;
@@ -831,6 +827,9 @@ public:
 	};
 };
 
+template <int IndexSize> class TSizedDefaultAllocator : public TSizedHeapAllocator<IndexSize> { public: typedef TSizedHeapAllocator<IndexSize> Typedef; };
+using FDefaultAllocator = TSizedDefaultAllocator<32>;
+
 // Define the ResizeAllocation functions with the regular allocator as exported to avoid bloat
 #if !IS_MERGEDMODULES && !defined(UE_HEADER_UNITS)
 extern template CORE_API FORCENOINLINE void TSizedHeapAllocator<32, FMemory>::ForAnyElementType::ResizeAllocation(SizeType CurrentNum, SizeType NewMax, SIZE_T NumBytesPerElement);
@@ -920,6 +919,9 @@ public:
 
 		void ResizeAllocation(SizeType CurrentNum, SizeType NewMax,SIZE_T NumBytesPerElement)
 		{
+			// Make sure the number of live elements is still withing the allocation since we only memmove, not destruct
+			checkSlow(CurrentNum <= NewMax);
+
 			// Check if the new allocation will fit in the inline data area.
 			if(NewMax <= NumInlineElements)
 			{
@@ -1113,6 +1115,9 @@ public:
 
 		void ResizeAllocation(SizeType CurrentNum, SizeType NewMax,SIZE_T NumBytesPerElement)
 		{
+			// Make sure the number of live elements is still withing the allocation since we only memmove, not destruct
+			checkSlow(CurrentNum <= NewMax);
+
 			// Check if the new allocation will fit in the inline data area.
 			if(NewMax <= NumInlineElements)
 			{
@@ -1328,6 +1333,8 @@ public:
 #define NumBitsPerDWORD ((int32)32)
 #define NumBitsPerDWORDLogTwo ((int32)5)
 
+class FDefaultBitArrayAllocator;
+
 //
 // Sparse array allocation definitions
 //
@@ -1383,6 +1390,79 @@ public:
 	typedef TFixedAllocator<InlineBitArrayDWORDs> BitArrayAllocator;
 };
 
+//
+// Compact set allocation definitions - It operates entirely with uint8 because it includes its lookup table inline with the overall allocation
+//
+
+namespace UE::Core::Private::CompactSetAllocatorHelpers
+{
+	template <uint32 NumInlineElements, int32 ElementSize>
+	constexpr int32 CalculateRequiredBytes()
+	{
+		constexpr uint32 TypeSize = 1 + (NumInlineElements > 0xff) + (NumInlineElements > 0xffff) * 2;
+		constexpr uint32 HashSize = NumInlineElements < 8 ? 4 : FGenericPlatformMath::RoundUpToPowerOfTwo((NumInlineElements / 2) + 1);
+		return Align(NumInlineElements*ElementSize, 4) + 4 + (NumInlineElements + HashSize) * TypeSize;
+	}
+}
+
+/** Regular iterator */
+template<typename InElementAllocator = FDefaultAllocator>
+struct TCompactSetAllocator
+{
+	/** CompactSets use raw uint8, so we need extra alignment information for inline allocators */
+	template <typename ElementType>
+	struct AllocatorAlignment
+	{
+		static constexpr size_t Value = alignof(typename InElementAllocator::template ForElementType<uint8>);
+	};
+
+	template<int32 ElementSize>
+	using ElementAllocator = InElementAllocator;
+};
+
+template<typename InElementAllocator>
+struct TAllocatorTraits<TCompactSetAllocator<InElementAllocator>> : TAllocatorTraitsBase<TCompactSetAllocator<InElementAllocator>>
+{
+	enum
+	{
+		SupportsFreezeMemoryImage = TAllocatorTraits<InElementAllocator>::SupportsFreezeMemoryImage,
+	};
+};
+
+/** An inline allocator that allows sizing of the inline allocations for a set number of elements. */
+template<uint32 NumInlineElements, typename SecondaryAllocator = TCompactSetAllocator<>>
+struct TInlineCompactSetAllocator
+{
+	template<int32 ElementSize>
+	using ElementAllocator = TInlineAllocator<UE::Core::Private::CompactSetAllocatorHelpers::CalculateRequiredBytes<NumInlineElements, ElementSize>(),typename SecondaryAllocator::template ElementAllocator<ElementSize>>;
+
+	/** CompactSets use raw uint8, so we need extra alignment information for inline allocators */
+	template <typename ElementType>
+	struct AllocatorAlignment
+	{
+		static constexpr size_t ElementAlignof = alignof(ElementType);
+		static constexpr size_t AllocatorAlignof = alignof(typename ElementAllocator<sizeof(ElementType)>::template ForElementType<uint8>);
+		static constexpr size_t Value = FGenericPlatformMath::Max(ElementAlignof, AllocatorAlignof);
+	};
+};
+
+/** An inline allocator that doesn't have any secondary storage. */
+template <uint32 NumInlineElements>
+struct TFixedCompactSetAllocator
+{
+	template<int32 ElementSize>
+	using ElementAllocator = TFixedAllocator<UE::Core::Private::CompactSetAllocatorHelpers::CalculateRequiredBytes<NumInlineElements, ElementSize>()>;
+	
+	/** CompactSets use raw uint8, so we need extra alignment information for inline allocators */
+	template <typename ElementType>
+	struct AllocatorAlignment
+	{
+		static constexpr size_t ElementAlignof = alignof(ElementType);
+		static constexpr size_t AllocatorAlignof = alignof(typename ElementAllocator<sizeof(ElementType)>::template ForElementType<uint8>);
+		static constexpr size_t Value = FGenericPlatformMath::Max(ElementAlignof, AllocatorAlignof);
+	};
+};
+
 
 
 //
@@ -1403,7 +1483,7 @@ template<
 	uint32   BaseNumberOfHashBuckets              = DEFAULT_BASE_NUMBER_OF_HASH_BUCKETS,
 	uint32   MinNumberOfHashedElements            = DEFAULT_MIN_NUMBER_OF_HASHED_ELEMENTS
 	>
-class TSetAllocator
+class TSparseSetAllocator
 {
 public:
 
@@ -1418,8 +1498,8 @@ public:
 		return 1;
 	}
 
-	typedef InSparseArrayAllocator SparseArrayAllocator;
-	typedef InHashAllocator        HashAllocator;
+	using SparseArrayAllocator = InSparseArrayAllocator;
+	using HashAllocator = InHashAllocator;
 };
 
 template<
@@ -1429,8 +1509,8 @@ template<
 	uint32   BaseNumberOfHashBuckets,
 	uint32   MinNumberOfHashedElements
 >
-struct TAllocatorTraits<TSetAllocator<InSparseArrayAllocator, InHashAllocator, AverageNumberOfElementsPerHashBucket, BaseNumberOfHashBuckets, MinNumberOfHashedElements>> :
-	TAllocatorTraitsBase<TSetAllocator<InSparseArrayAllocator, InHashAllocator, AverageNumberOfElementsPerHashBucket, BaseNumberOfHashBuckets, MinNumberOfHashedElements>>
+struct TAllocatorTraits<TSparseSetAllocator<InSparseArrayAllocator, InHashAllocator, AverageNumberOfElementsPerHashBucket, BaseNumberOfHashBuckets, MinNumberOfHashedElements>> :
+	TAllocatorTraitsBase<TSparseSetAllocator<InSparseArrayAllocator, InHashAllocator, AverageNumberOfElementsPerHashBucket, BaseNumberOfHashBuckets, MinNumberOfHashedElements>>
 {
 	enum
 	{
@@ -1441,11 +1521,11 @@ struct TAllocatorTraits<TSetAllocator<InSparseArrayAllocator, InHashAllocator, A
 /** An inline set allocator that allows sizing of the inline allocations for a set number of elements. */
 template<
 	uint32   NumInlineElements,
-	typename SecondaryAllocator                   = TSetAllocator<TSparseArrayAllocator<FDefaultAllocator,FDefaultAllocator>,FDefaultAllocator>,
+	typename SecondaryAllocator                   = TSparseSetAllocator<TSparseArrayAllocator<FDefaultAllocator,FDefaultAllocator>,FDefaultAllocator>,
 	uint32   AverageNumberOfElementsPerHashBucket = DEFAULT_NUMBER_OF_ELEMENTS_PER_HASH_BUCKET,
 	uint32   MinNumberOfHashedElements            = DEFAULT_MIN_NUMBER_OF_HASHED_ELEMENTS
 	>
-class TInlineSetAllocator
+class TInlineSparseSetAllocator
 {
 private:
 
@@ -1472,8 +1552,8 @@ public:
 		return NumDesiredHashBuckets;
 	}
 
-	typedef TInlineSparseArrayAllocator<NumInlineElements,typename SecondaryAllocator::SparseArrayAllocator> SparseArrayAllocator;
-	typedef TInlineAllocator<NumInlineHashBuckets,typename SecondaryAllocator::HashAllocator>                HashAllocator;
+	using SparseArrayAllocator = TInlineSparseArrayAllocator<NumInlineElements,typename SecondaryAllocator::SparseArrayAllocator>;
+	using HashAllocator = TInlineAllocator<NumInlineHashBuckets,typename SecondaryAllocator::HashAllocator>;
 };
 
 /** An inline set allocator that doesn't have any secondary storage. */
@@ -1482,7 +1562,7 @@ template<
 	uint32 AverageNumberOfElementsPerHashBucket = DEFAULT_NUMBER_OF_ELEMENTS_PER_HASH_BUCKET,
 	uint32 MinNumberOfHashedElements            = DEFAULT_MIN_NUMBER_OF_HASHED_ELEMENTS
 	>
-class TFixedSetAllocator
+class TFixedSparseSetAllocator
 {
 private:
 
@@ -1509,8 +1589,8 @@ public:
 		return NumDesiredHashBuckets;
 	}
 
-	typedef TFixedSparseArrayAllocator<NumInlineElements> SparseArrayAllocator;
-	typedef TFixedAllocator<NumInlineHashBuckets>         HashAllocator;
+	using SparseArrayAllocator = TFixedSparseArrayAllocator<NumInlineElements>;
+	using HashAllocator = TFixedAllocator<NumInlineHashBuckets>;
 };
 
 
@@ -1521,16 +1601,38 @@ public:
  * 'forward' these TAllocatorTraits specializations below.
  */
 
-template <int IndexSize> class TSizedDefaultAllocator : public TSizedHeapAllocator<IndexSize> { public: typedef TSizedHeapAllocator<IndexSize> Typedef; };
-
-class FDefaultSetAllocator         : public TSetAllocator<>         { public: typedef TSetAllocator<>         Typedef; };
+class FDefaultSparseSetAllocator   : public TSparseSetAllocator<>   { public: typedef TSparseSetAllocator<>   Typedef; };
+class FDefaultCompactSetAllocator  : public TCompactSetAllocator<>  { public: typedef TCompactSetAllocator<>  Typedef; };
 class FDefaultBitArrayAllocator    : public TInlineAllocator<4>     { public: typedef TInlineAllocator<4>     Typedef; };
 class FDefaultSparseArrayAllocator : public TSparseArrayAllocator<> { public: typedef TSparseArrayAllocator<> Typedef; };
 
+#ifndef UE_USE_COMPACT_SET_AS_DEFAULT
+#define UE_USE_COMPACT_SET_AS_DEFAULT 0
+#endif
+
+#if UE_USE_COMPACT_SET_AS_DEFAULT
+class FDefaultSetAllocator         : public TCompactSetAllocator<>  { public: typedef TCompactSetAllocator<>  Typedef; };
+template<typename InSparseArrayAllocator = TSparseArrayAllocator<>, typename InHashAllocator = TInlineAllocator<1,FDefaultAllocator>, uint32... N>
+class TSetAllocator        		   : public TCompactSetAllocator<InHashAllocator>  { public: typedef TCompactSetAllocator<InHashAllocator>  Typedef; };
+template <uint32 N, typename S = TCompactSetAllocator<>, uint32... NN>
+class TInlineSetAllocator          : public TInlineCompactSetAllocator<N, S>  { public: typedef TInlineCompactSetAllocator<N, S>  Typedef; };
+template <uint32 N, uint32... Y>
+class TFixedSetAllocator           : public TFixedCompactSetAllocator<N>  { public: typedef TFixedCompactSetAllocator<N>  Typedef; };
+#else
+class FDefaultSetAllocator         : public TSparseSetAllocator<>         { public: typedef TSparseSetAllocator<>         Typedef; };
+template<typename InSparseArrayAllocator = TSparseArrayAllocator<>, typename InHashAllocator = TInlineAllocator<1,FDefaultAllocator>, uint32... N>
+class TSetAllocator         	   : public TSparseSetAllocator<InSparseArrayAllocator, InHashAllocator, N...>  { public: typedef TSparseSetAllocator<InSparseArrayAllocator, InHashAllocator, N...>  Typedef; };
+template<uint32 N, typename S = TSparseSetAllocator<TSparseArrayAllocator<FDefaultAllocator,FDefaultAllocator>,FDefaultAllocator>, uint32... NN>
+class TInlineSetAllocator          : public TInlineSparseSetAllocator<N, S, NN...>  { public: typedef TInlineSparseSetAllocator<N, S, NN...>  Typedef; };
+template <uint32... N>
+class TFixedSetAllocator           : public TFixedSparseSetAllocator<N...>  { public: typedef TFixedSparseSetAllocator<N...>  Typedef; };
+#endif
+
 template <int IndexSize> struct TAllocatorTraits<TSizedDefaultAllocator<IndexSize>> : TAllocatorTraits<typename TSizedDefaultAllocator<IndexSize>::Typedef> {};
 
 template <> struct TAllocatorTraits<FDefaultAllocator>            : TAllocatorTraits<typename FDefaultAllocator           ::Typedef> {};
 template <> struct TAllocatorTraits<FDefaultSetAllocator>         : TAllocatorTraits<typename FDefaultSetAllocator        ::Typedef> {};
+template <> struct TAllocatorTraits<FDefaultSparseSetAllocator>   : TAllocatorTraits<typename FDefaultSparseSetAllocator  ::Typedef> {};
 template <> struct TAllocatorTraits<FDefaultBitArrayAllocator>    : TAllocatorTraits<typename FDefaultBitArrayAllocator   ::Typedef> {};
 template <> struct TAllocatorTraits<FDefaultSparseArrayAllocator> : TAllocatorTraits<typename FDefaultSparseArrayAllocator::Typedef> {};
 
diff --git a/Engine/Source/Runtime/Core/Public/Containers/ContainersFwd.h b/Engine/Source/Runtime/Core/Public/Containers/ContainersFwd.h
--- a/Engine/Source/Runtime/Core/Public/Containers/ContainersFwd.h
+++ b/Engine/Source/Runtime/Core/Public/Containers/ContainersFwd.h
@@ -8,9 +8,13 @@
 
 /// @cond DOXYGEN_WARNINGS
 template<int IndexSize> class TSizedDefaultAllocator;
+template<typename ElementType,bool bInAllowDuplicateKeys = false> struct DefaultKeyFuncs;
+
 using FDefaultAllocator = TSizedDefaultAllocator<32>;
 using FDefaultAllocator64 = TSizedDefaultAllocator<64>;
 class FDefaultSetAllocator;
+class FDefaultCompactSetAllocator;
+class FDefaultSparseSetAllocator;
 
 class FString;
 class FAnsiString;
@@ -65,13 +69,18 @@ template<typename T, typename SizeType = int32> using TConstArrayView = TArrayVi
 template<typename T> using TConstArrayView64 = TConstArrayView<T, int64>;
 template<typename T> class TTransArray;
 template<typename InKeyType, typename InValueType, bool bInAllowDuplicateKeys> struct TDefaultMapHashableKeyFuncs;
+template<typename InKeyType, typename InValueType, typename SetAllocator = FDefaultSetAllocator, typename KeyFuncs = TDefaultMapHashableKeyFuncs<InKeyType, InValueType, false> > class TSparseMap;
+template<typename InKeyType, typename InValueType, typename SetAllocator = FDefaultSetAllocator, typename KeyFuncs = TDefaultMapHashableKeyFuncs<InKeyType, InValueType, false> > class TCompactMap;
 template<typename InKeyType, typename InValueType, typename SetAllocator = FDefaultSetAllocator, typename KeyFuncs = TDefaultMapHashableKeyFuncs<InKeyType, InValueType, false> > class TMap;
+template<typename KeyType, typename ValueType, typename SetAllocator = FDefaultSetAllocator, typename KeyFuncs = TDefaultMapHashableKeyFuncs<KeyType, ValueType, true > > class TSparseMultiMap;
+template<typename KeyType, typename ValueType, typename SetAllocator = FDefaultSetAllocator, typename KeyFuncs = TDefaultMapHashableKeyFuncs<KeyType, ValueType, true > > class TCompactMultiMap;
 template<typename KeyType, typename ValueType, typename SetAllocator = FDefaultSetAllocator, typename KeyFuncs = TDefaultMapHashableKeyFuncs<KeyType, ValueType, true > > class TMultiMap;
 template <typename T = void > struct TLess;
 template <typename> struct TTypeTraits;
 template<typename InKeyType, typename InValueType, typename ArrayAllocator = FDefaultAllocator, typename SortPredicate = TLess<typename TTypeTraits<InKeyType>::ConstPointerType> > class TSortedMap;
 template <typename InElementType, typename ArrayAllocator = FDefaultAllocator, typename SortPredicate = TLess<InElementType>> class TSortedSet;
-template<typename ElementType,bool bInAllowDuplicateKeys = false> struct DefaultKeyFuncs;
+template<typename InElementType, typename KeyFuncs = DefaultKeyFuncs<InElementType>, typename Allocator = FDefaultSparseSetAllocator> class TSparseSet;
+template<typename InElementType, typename KeyFuncs = DefaultKeyFuncs<InElementType>, typename Allocator = FDefaultCompactSetAllocator> class TCompactSet;
 template<typename InElementType, typename KeyFuncs = DefaultKeyFuncs<InElementType>, typename Allocator = FDefaultSetAllocator> class TSet;
 template<typename InElementType, typename InSizeType = int32> class TStridedView;
 template<typename T, typename SizeType = int32> using TConstStridedView = TStridedView<const T, SizeType>;
