PlainProps: Split out a private PlainPropsUObjectBindingsInternal.h and PlainPropsRoundTripTest.cpp from PlainPropsUObjectBindings.cpp.
Reduce normal wall compile time (with PCH) for PlainPropsUObjectBindings.cpp from 4.5s to 3.5s (1.7s for PlainPropsRoundTripTest.cpp).
Reduce -SingleFile wall compile time (no PCH) for PlainPropsUObjectBindings.cpp from 9.5s to 7.3s (6.5s for PlainPropsRoundTripTest.cpp).

#rnx

diff --git a/Engine/Plugins/Experimental/PlainPropsUObject/Source/Private/PlainPropsRoundtripTest.cpp b/Engine/Plugins/Experimental/PlainPropsUObject/Source/Private/PlainPropsRoundtripTest.cpp
new file mode 100644
--- /dev/null
+++ b/Engine/Plugins/Experimental/PlainPropsUObject/Source/Private/PlainPropsRoundtripTest.cpp
@@ -0,0 +1,1654 @@
+// Copyright Epic Games, Inc. All Rights Reserved.
+
+#include "PlainPropsRoundtripTest.h"
+#include "PlainPropsUObjectBindingsInternal.h"
+#include "PlainPropsBuildSchema.h"
+#include "PlainPropsLoadMember.h"
+#include "PlainPropsUObjectRuntime.h"
+#include "PlainPropsParse.h"
+#include "PlainPropsPrint.h"
+#include "PlainPropsSpecify.h"
+#include "PlainPropsWrite.h"
+#include "HAL/FileManager.h"
+#include "JsonObjectConverter.h"
+#include "Logging/StructuredLog.h"
+#include "Misc/Filehelper.h"
+#include "Misc/PackageName.h"
+#include "Misc/Paths.h"
+#include "Serialization/Archive.h"
+#include "Serialization/BufferWriter.h"
+#include "Serialization/LargeMemoryReader.h"
+#include "Serialization/MemoryReader.h"
+#include "Serialization/MemoryWriter.h"
+#include "Serialization/UnversionedPropertySerializationInternal.h"
+#include "Templates/UniquePtr.h"
+#include "UObject/LinkerLoad.h"
+#include "UObject/LinkerSave.h"
+#include "UObject/NameBatchSerialization.h"
+#include "UObject/Package.h"
+#include "UObject/SavePackage.h"
+#include "UObject/TextProperty.h"
+#include "UObject/UObjectIterator.h"
+
+#define PP_WITH_ENGINE_HOOKS 0
+
+namespace PlainProps::UE
+{
+
+struct FMemoryPropertyBatch
+{
+	TArray<FText> Texts; // Tricky to serialize intrusively
+
+	static FSoleMemberSpec Spec(FName*, FDeclId Id)
+	{
+		return {Id, GUE.Members.Id, Specify<uint64>()};
+	}
+
+	static void Save(FMemberBuilder& Out, FName In, const FSaveContext&)
+	{
+		Out.Add(GUE.Members.Id, FSensitiveName(In).ToUnstableInt());
+	}
+
+	static void Load(FName& Out, FStructLoadView In)
+	{
+		Out = FSensitiveName::FromUnstableInt(LoadSole<FSensitiveName::IntType>(In)).ToName();
+	}
+
+	static FSoleMemberSpec Spec(FText*, FDeclId Id)
+	{
+		return {Id, GUE.Members.Id, Specify<int32>()};
+	}
+
+	void Save(FMemberBuilder& Out, const FText& In, const FSaveContext&)
+	{
+		Out.Add(GUE.Members.Id, Texts.Num());
+		Texts.Add(In);
+	}
+
+	void Load(FText& Out, FStructLoadView In) const
+	{
+		Out = Texts[LoadSole<int32>(In)];
+	}
+
+	static FSoleMemberSpec Spec(FObjectHandle*, FDeclId Id)
+	{
+		return {Id, GUE.Members.Id, Specify<uint64>()};
+	}
+
+	static void Save(FMemberBuilder& Out, FObjectHandle In, const FSaveContext&)
+	{
+		static_assert(sizeof(In) == sizeof(uint64));
+		Out.Add(GUE.Members.Id, reinterpret_cast<const uint64&>(In));
+	}
+
+	static void Load(FObjectHandle& Out, FStructLoadView In)
+	{
+		LoadSole<uint64>(&Out, In);
+	}
+
+	static FSoleMemberSpec Spec(FWeakObjectPtr*, FDeclId Id)
+	{
+		return {Id, GUE.Members.Id, Specify<uint64>()};
+	}
+
+	static void Save(FMemberBuilder& Out, const FWeakObjectPtr& In, const FSaveContext&)
+	{
+		// Save ObjectSerialNumber + ObjectIndex a single uint64
+		static_assert(sizeof(FWeakObjectPtr) == sizeof(uint64));
+		Out.Add(GUE.Members.Id, reinterpret_cast<const uint64&>(In));
+	}
+
+	static void Load(FWeakObjectPtr& Out, FStructLoadView In)
+	{
+		LoadSole<uint64>(&Out, In);
+	}
+
+	static FSoleMemberSpec Spec(FSoftObjectPtr*, FDeclId Id)
+	{
+		return {Id, GUE.Members.Id, FMemberSpec(GUE.Structs.SoftObjectPath)};
+	}
+
+	static void Save(FMemberBuilder& Out, const FSoftObjectPtr& In, const FSaveContext& Ctx)
+	{
+		FBuiltStruct* SoftPath = SaveStruct(&In.GetUniqueID(), GUE.Structs.SoftObjectPath, Ctx);
+		Out.AddStruct(GUE.Members.Id, GUE.Structs.SoftObjectPath, SoftPath);
+	}
+
+	static void Load(FSoftObjectPtr& Out, FStructLoadView In)
+	{
+		Out.ResetWeakPtr();
+		LoadSoleStruct(&Out.GetUniqueID(), In);
+	}
+	
+	static FSoleMemberSpec Spec(FLazyObjectPtr*, FDeclId Id)
+	{
+		return {Id, GUE.Members.Id, FMemberSpec(GUE.Structs.Guid)};
+	}
+	
+	static void Save(FMemberBuilder& Out, const FLazyObjectPtr& In, const FSaveContext& Ctx)
+	{
+		FBuiltStruct* Guid = SaveStruct(&In.GetUniqueID(), GUE.Structs.Guid, Ctx);
+		Out.AddStruct(GUE.Members.Id, GUE.Structs.Guid, Guid);
+	}
+
+	static void Load(FLazyObjectPtr& Out, FStructLoadView In)
+	{
+		Out.ResetWeakPtr();
+		LoadSoleStruct(&Out.GetUniqueID(), In);
+	}
+};
+
+inline bool DiffProperty(FName A, FName B) { return !A.IsEqual(B, ENameCase::CaseSensitive); }
+inline bool DiffProperty(const FText& A, const FText& B) { return !FTextProperty::Identical_Implementation(A, B, 0); }
+inline bool DiffProperty(FObjectHandle A, FObjectHandle B) { return A != B; }
+inline bool DiffProperty(const FWeakObjectPtr& A, const FWeakObjectPtr& B) { return A != B; }
+inline bool DiffProperty(const FSoftObjectPtr& A, const FSoftObjectPtr& B) { return A != B; }
+inline bool DiffProperty(const FLazyObjectPtr& A, const FLazyObjectPtr& B) { return A != B; }
+
+template<class Type, class BatchType>
+struct TCustomPropertyBinding final : ICustomBinding
+{
+	TCustomPropertyBinding(BatchType& InBatch) : Batch(InBatch) {}
+
+	BatchType& Batch;
+
+	virtual void SaveCustom(FMemberBuilder& Dst, const void* Src, const void* Default, const FSaveContext& Ctx) override
+	{
+		if (!Default || DiffCustom(Src, Default, Ctx))
+		{
+			Batch.Save(Dst, *static_cast<const Type*>(Src), Ctx);
+		}
+	}
+
+	virtual void LoadCustom(void* Dst, FStructLoadView Src, ECustomLoadMethod) const override
+	{
+		Batch.Load(*static_cast<Type*>(Dst), Src);
+	}
+
+	virtual bool DiffCustom(const void* A, const void* B, const FBindContext&) const override
+	{
+		return DiffProperty(*static_cast<const Type*>(A), *static_cast<const Type*>(B));
+	}
+};
+
+template<class BatchType>
+struct TCustomPropertyBindings
+{
+	TCustomPropertyBindings(BatchType& Batch, const FCustomBindings& Underlay)
+	: Overlay(Underlay)
+	, Name(Batch)
+	, Text(Batch)
+	, ObjectPtr(Batch)
+	, SoftObjectPtr(Batch)
+	, WeakObjectPtr(Batch)
+	, LazyObjectPtr(Batch)
+	{
+		Bind(GUE.Structs.Name, Name);
+		Bind(GUE.Structs.Text, Text);
+		Bind(GUE.Structs.ClassPtr, ObjectPtr); // TSubclassOf<> is essentially a TObjectPtr
+		Bind(GUE.Structs.ObjectPtr, ObjectPtr);
+		Bind(GUE.Structs.SoftObjectPtr, SoftObjectPtr);
+		Bind(GUE.Structs.WeakObjectPtr, WeakObjectPtr);
+		Bind(GUE.Structs.LazyObjectPtr, LazyObjectPtr);
+	}
+
+	template<class Type>
+	void Bind(FDualStructId Id, TCustomPropertyBinding<Type, BatchType>& Binding)
+	{
+		Overlay.BindStruct(Id, Binding, Declare(BatchType::Spec((Type*)nullptr, Id)), {});
+	}
+
+	FCustomBindingsOverlay								Overlay;
+	TCustomPropertyBinding<FName, BatchType>			Name;
+	TCustomPropertyBinding<FText, BatchType>			Text;
+	TCustomPropertyBinding<FObjectHandle, BatchType>	ObjectPtr;
+	TCustomPropertyBinding<FSoftObjectPtr, BatchType>	SoftObjectPtr;
+	TCustomPropertyBinding<FWeakObjectPtr, BatchType>	WeakObjectPtr;
+	TCustomPropertyBinding<FLazyObjectPtr, BatchType>	LazyObjectPtr;
+};
+
+////////////////////////////////////////////////////////////////////////////////////////////////
+
+struct FMemoryBatch
+{
+	TArray64<uint8>				Data;
+	TArray<FStructId>			RuntimeIds; // To avoid reindexing schema FType
+	FMemoryPropertyBatch		Properties; // Contains FTexts referenced from Data
+};
+
+static constexpr uint32 Magics[] = { 0xFEEDF00D, 0xABCD1234, 0xDADADAAA, 0x99887766, 0xF0F1F2F3 };
+volatile static const UObject* GDebugNoteObject;
+
+class FBatchSaver
+{
+public:
+	FBatchSaver(FCustomBindings& Customs, int32 NumReserve)
+	: FlatCtx{{GUE.Schemas, Customs}, Scratch, nullptr}
+	, DeltaCtx{{GUE.Schemas, Customs}, Scratch, &GUE.Defaults}
+	{
+		SavedObjects.Reserve(NumReserve);
+		ObjectMap.Reserve(NumReserve);
+	}
+
+	// If Arch, delta serialize the matching part of the object inheritance chain
+	FBuiltStruct* SaveImpl(FBindId Id, const UStruct* Class, const UObject* Object, const UStruct* ArchClass, const UObject* Arch)
+	{
+		if (ArchClass == nullptr)
+		{
+			return SaveStruct(Object, Id, FlatCtx);
+		}
+		if (Class == ArchClass)
+		{
+			return SaveStructDelta(Object, Arch, Id, DeltaCtx);
+		}
+
+		Class = SkipEmptyBases(Class->GetInheritanceSuper());
+		FBindId SuperId = GUE.Names.IndexBindId(IndexType(Class));
+		FBuiltStruct* Super = SaveImpl(SuperId, Class, Object, ArchClass, Arch);
+		return SaveStructWithSuper(Object, Super, Id, FlatCtx);
+	}
+
+	void Save(FBindId Id, const UObject* Object, const UObject* Arch)
+	{
+		FBuiltStruct* Built;
+		if (Arch)
+		{
+			const UStruct* Class = SkipEmptyBases(Object->GetClass());
+			const UStruct* ArchClass = SkipEmptyBases(Arch->GetClass());
+			Built = SaveImpl(Id, Class, Object, ArchClass, Arch);
+		}
+		else
+		{
+			Built = SaveImpl(Id, Object->GetClass(), Object, nullptr, nullptr);
+		}
+		ObjectMap.Add(SavedObjects.Num());
+		SavedObjects.Emplace(Id, Built, Object);
+	}
+	
+	void Skip(FBindId Id, const UObject* Object)
+	{
+		ObjectMap.Add(~0u);
+	}
+
+	// Use existing SaveNameBatch() API for now, can be optimized later
+	void WriteStableNames(TArray64<uint8>& Out, TConstArrayView<FNameId> UsedNames) const
+	{
+		TArray<FDisplayNameEntryId> Entries;
+		TArray<int32> Numbers;
+		Entries.Reserve(UsedNames.Num());
+		Numbers.Reserve(UsedNames.Num());
+		for (FNameId NameId : UsedNames)
+		{
+			FSensitiveName SensitiveName = GUE.Names.ResolveName(NameId);
+			FName Name = SensitiveName.ToName();
+			Entries.Add(FDisplayNameEntryId(Name));
+			Numbers.Add(Name.GetNumber());
+		}
+
+		FBufferWriter NameMapArchive(nullptr, 0, EBufferWriterFlags::AllowResize | EBufferWriterFlags::TakeOwnership);
+		SaveNameBatch(Entries, NameMapArchive);
+
+		WriteNumAndArray(Out, TArrayView<uint8>(static_cast<uint8*>(NameMapArchive.GetWriterData()), static_cast<int32>(NameMapArchive.Tell())));
+		WriteNumAndArray(Out, Numbers);
+	}
+
+	TArray64<uint8>	Write(TArray<FStructId>* OutRuntimeIds = nullptr) const
+	{
+		ESchemaFormat Format = OutRuntimeIds ? ESchemaFormat::InMemoryNames : ESchemaFormat::StableNames;
+
+		// Build partial schemas
+		const FBindDeclarations Types(GUE.Enums, FlatCtx.Customs, GUE.Schemas);
+		FSchemasBuilder SchemaBuilders(GUE.Names, Types, Scratch, Format);
+		for (const FSavedObject& Object : SavedObjects)
+		{
+			GDebugNoteObject = Object.Input;
+			SchemaBuilders.NoteStructAndMembers(Object.Id, *Object.Built);
+		}
+		GDebugNoteObject = nullptr;
+		FBuiltSchemas Schemas = SchemaBuilders.Build();
+
+		// Save schema ids on the side when using InMemoryNames
+		if (OutRuntimeIds)
+		{
+			*OutRuntimeIds = ExtractRuntimeIds(Schemas);
+		}
+
+		FWriter Writer(GUE.Names, Types, Schemas, Format);
+		TArray64<uint8> Out;
+
+		WriteInt(Out, Magics[0]);
+		// Write out FNames when using StableNames
+		if (!OutRuntimeIds)
+		{
+			WriteStableNames(Out, Writer.GetUsedNames());
+		}
+
+		// Write schemas
+		WriteInt(Out, Magics[1]);
+		WriteAlignmentPadding<uint32>(Out);
+		TArray64<uint8> Tmp;
+		Writer.WriteSchemas(/* Out */ Tmp);
+		WriteNumAndArray(Out, TArrayView<const uint8, int64>(Tmp));
+		Tmp.Reset();
+
+		// Write objects
+		WriteInt(Out, Magics[2]);
+		for (const FSavedObject& Object : SavedObjects)
+		{
+			WriteInt(/* out */ Tmp, Magics[3]);
+			WriteInt(/* out */ Tmp, Writer.GetWriteId(Object.Id).Get().Idx);
+			Writer.WriteMembers(/* out */ Tmp, Object.Id, *Object.Built);
+			WriteSkippableSlice(Out, Tmp);
+			Tmp.Reset();
+		}
+
+		// Write object terminator
+		WriteSkippableSlice(Out, TConstArrayView64<uint8>());
+		WriteInt(Out, Magics[4]);
+
+		// Write object map that allows for empty objects/(indices) without built data
+		WriteNumAndArray(Out, ObjectMap);
+		
+		return Out;
+	}
+
+private:
+	struct FSavedObject
+	{
+		FBindId				Id;
+		FBuiltStruct*		Built;
+		const UObject*		Input; // For debug
+	};
+
+	TArray<FSavedObject>		SavedObjects;
+	TArray<uint32>				ObjectMap;
+	mutable FScratchAllocator	Scratch;
+	FSaveContext				FlatCtx;
+	FSaveContext				DeltaCtx;
+
+	template<typename ArrayType>
+	static void WriteNumAndArray(TArray64<uint8>& Out, const ArrayType& Items)
+	{
+		WriteInt(Out, IntCastChecked<uint32>(Items.Num()));
+		WriteArray(Out, Items);
+	}
+};
+
+class FMemoryBatchLoader
+{
+public:
+	FMemoryBatchLoader(const FCustomBindings& Customs, FMemoryView Data, TConstArrayView<FStructId> RuntimeIds)
+	{
+		// Read and mount schemas
+		FByteReader It(Data);
+		verify(It.Grab<uint32>() == Magics[0]);
+		verify(It.Grab<uint32>() == Magics[1]);
+		It.SkipAlignmentPadding<uint32>();
+		uint32 SchemasSize = It.Grab<uint32>();
+		const FSchemaBatch* SavedSchemas = ValidateSchemas(It.GrabSlice(SchemasSize));
+		verify(It.Grab<uint32>() == Magics[2]);
+		
+		FSchemaBatchId Batch = MountReadSchemas(SavedSchemas);
+
+		// Read objects
+ 		while (uint64 NumBytes = It.GrabVarIntU())
+		{	
+			FByteReader ObjIt(It.GrabSlice(NumBytes));
+			verify(ObjIt.Grab<uint32>() == Magics[3]);
+			FStructSchemaId Id = { ObjIt.Grab<uint32>() };
+			check(Id != FStructSchemaId{});
+			Objects.Add({ { Id, Batch }, ObjIt });
+		}
+		
+		verify(It.Grab<uint32>() == Magics[4]);
+		verify(GrabNumAndArray<uint32>(It).Num() == Objects.Num());
+		It.CheckEmpty();
+		check(!Objects.IsEmpty());
+
+		// Finally create load plans
+		Plans = CreateLoadPlans(Batch, Customs, GUE.Schemas, RuntimeIds, ESchemaFormat::InMemoryNames);
+	}
+
+	~FMemoryBatchLoader()
+	{
+		check(LoadIdx == Objects.Num()); // Test should load all saved objects
+		Plans.Reset();
+		UnmountReadSchemas(Objects[0].Schema.Batch);
+	}
+
+	FORCENOINLINE void Load(UObject* Dst)
+	{
+		FStructView In = Objects[LoadIdx];
+		LoadStruct(Dst, In.Values, In.Schema.Id, *Plans);
+		++LoadIdx;
+	}
+
+	FORCENOINLINE void Reload(UObject* Dst, int32 ReloadIdx)
+	{
+		FStructView In = Objects[ReloadIdx];
+		LoadStruct(Dst, In.Values, In.Schema.Id, *Plans);
+	}
+
+private:
+	FLoadBatchPtr				Plans;
+	TArray<FStructView>			Objects;
+	int32						LoadIdx = 0;
+	
+	template<typename T>
+	static TConstArrayView<T> GrabNumAndArray(/* in-out */ FByteReader& It)
+	{
+		uint32 Num = It.Grab<uint32>();
+		return MakeArrayView(reinterpret_cast<const T*>(It.GrabBytes(Num * sizeof(T))), Num);
+	}
+};
+
+class FPersistentBatchLoader
+{
+public:
+	FPersistentBatchLoader(const FCustomBindings& Customs, FMemoryView Data)
+	{
+		// Read ids
+		FByteReader It(Data);
+		verify(It.Grab<uint32>() == Magics[0]);
+		ReadStableNames(It);
+
+		// Read and mount schemas
+		verify(It.Grab<uint32>() == Magics[1]);
+		It.SkipAlignmentPadding<uint32>();
+		uint32 SchemasSize = It.Grab<uint32>();
+		const FSchemaBatch* SavedSchemas = ValidateSchemas(It.GrabSlice(SchemasSize));
+		verify(It.Grab<uint32>() == Magics[2]);
+		
+		// Bind saved ids to runtime ids, make new schemas with new ids and mount them
+		FIdTranslator RuntimeIds(GUE.Names, MakeConstArrayView(Names), *SavedSchemas);
+		TranslatedSchemas = CreateTranslatedSchemas(*SavedSchemas, RuntimeIds.Translation);
+		FSchemaBatchId Batch = MountReadSchemas(TranslatedSchemas);
+
+		// Read objects
+ 		while (uint64 NumBytes = It.GrabVarIntU())
+		{	
+			FByteReader ObjIt(It.GrabSlice(NumBytes));
+			verify(ObjIt.Grab<uint32>() == Magics[3]);
+			FStructSchemaId Id = { ObjIt.Grab<uint32>() };
+			Objects.Add({ { Id, Batch }, ObjIt });
+		}
+		
+		verify(It.Grab<uint32>() == Magics[4]);
+		ObjectMap = GrabNumAndArray<uint32>(It);
+		check(ObjectMap.Num() >= Objects.Num());
+		check(!Objects.IsEmpty());
+
+		// Finally create load plans
+		TConstArrayView<FStructId> LoadStructIds = RuntimeIds.Translation.GetStructIds(NumStructSchemas(Batch));
+		Plans = CreateLoadPlans(Batch, Customs, GUE.Schemas, LoadStructIds, ESchemaFormat::StableNames);
+	}
+
+	void ReadStableNames(FByteReader& It)
+	{
+		TConstArrayView64<uint8> NameBatch = GrabNumAndArray<uint8>(It);
+		TConstArrayView<int32> NameNumbers = GrabNumAndArray<int32>(It);
+		FMemoryReaderView Ar(MakeArrayView(NameBatch.GetData(), NameBatch.Num()));
+		TArray<FDisplayNameEntryId> NameEntries = LoadNameBatch(Ar);
+		Names.Reserve(NameEntries.Num());
+		check(NameEntries.Num() == NameNumbers.Num());
+		for (int64 Idx = 0, Num = NameEntries.Num(); Idx < Num; ++Idx)
+		{
+			Names.Emplace(NameEntries[Idx].ToName(NameNumbers[Idx]));
+		}
+	}
+
+	~FPersistentBatchLoader()
+	{
+		Plans.Reset();
+		UnmountReadSchemas(Objects[0].Schema.Batch);
+	}
+
+	FORCENOINLINE void Load(UObject* Dst)
+	{
+		FStructView In = Objects[LoadIdx];
+		LoadStruct(Dst, In.Values, In.Schema.Id, *Plans);
+		++LoadIdx;
+	}
+
+	bool HasObject(int32 ObjectIndex)
+	{
+		return ObjectMap[ObjectIndex] != ~0u;
+	}
+
+	FORCENOINLINE void LoadObject(UObject* Dst, int32 ObjectIndex)
+	{
+		FStructView In = Objects[ObjectMap[ObjectIndex]];
+		LoadStruct(Dst, In.Values, In.Schema.Id, *Plans);
+	}
+
+private:
+	const FSchemaBatch*			TranslatedSchemas = nullptr;
+	TArray<FSensitiveName>		Names;
+	FLoadBatchPtr				Plans;
+	TArray<FStructView>			Objects;
+	TConstArrayView<uint32>		ObjectMap;
+	int32						LoadIdx = 0;
+
+	template<typename T>
+	static TConstArrayView<T> GrabNumAndArray(/* in-out */ FByteReader& It)
+	{
+		uint32 Num = It.Grab<uint32>();
+		return MakeArrayView(reinterpret_cast<const T*>(It.GrabBytes(Num * sizeof(T))), Num);
+	}
+};
+
+////////////////////////////////////////////////////////////////////////////////////////////////
+
+// Similar to PlainProps FMemoryBatch but for storing FArchive property serialization results
+struct FArchivedProperties
+{
+	TArray<uint8>			Data;
+	TArray<FText>			Texts;
+};
+
+static constexpr uint32 RoundtripPortFlags = PPF_UseDeprecatedProperties | PPF_ForceTaggedSerialization;
+
+// Match FMemoryPropertyBatch somewhat for a fair comparison, e.g. save FText on side and FName as integer
+class FPropertyWriter final : public FMemoryWriter
+{
+	TArray<FText>& Texts;
+public:
+	FPropertyWriter(FArchivedProperties& Out) 
+	: FMemoryWriter(Out.Data)
+	, Texts(Out.Texts)
+	{
+		SetPortFlags(RoundtripPortFlags);
+	}
+
+	FORCENOINLINE void WriteProperties(UObject* Object, UObject* Defaults)
+	{
+		UClass* Class = Object->GetClass();
+		Class->SerializeTaggedProperties(*this, reinterpret_cast<uint8*>(Object), Class, reinterpret_cast<uint8*>(Defaults));
+	}
+
+	virtual FArchive& operator<<(FText& Value) override
+	{
+		int32 Idx = INDEX_NONE;
+		if (!Value.IsEmpty())
+		{
+			Idx = Texts.Num();
+			Texts.Add(Value);
+		}
+		return WriteValue(Idx);
+	}
+	virtual FArchive& operator<<(FName& Value) override				{ return WriteValue(FSensitiveName(Value).ToUnstableInt());}
+	virtual FArchive& operator<<(UObject*& Value) override			{ return WriteValue(reinterpret_cast<uint64&>(Value)); }
+	virtual FArchive& operator<<(FObjectPtr& Value) override		{ return WriteValue(reinterpret_cast<uint64&>(Value)); }
+	virtual FArchive& operator<<(FWeakObjectPtr& Value) override	{ return WriteValue(reinterpret_cast<uint64&>(Value)); }
+	virtual FArchive& operator<<(FLazyObjectPtr& Value) override	{ return WriteValue(Value.GetUniqueID()); }
+	virtual FArchive& operator<<(FSoftObjectPtr& Value) override	{ Value.GetUniqueID().SerializePath(*this); return *this; }
+	virtual FArchive& operator<<(FSoftObjectPath& Value) override	{ Value.SerializePath(*this); return *this; }
+	virtual FString GetArchiveName() const override					{ return "FPropertyWriter"; }
+
+	template<typename T>
+	FArchive& WriteValue(T Value)
+	{
+		return *this << Value;
+	}
+};
+
+class FPropertyReader final : public FMemoryReader
+{
+	const TArray<FText>& Texts;
+public:
+	FPropertyReader(const FArchivedProperties& Out)
+	: FMemoryReader(Out.Data)
+	, Texts(Out.Texts)
+	{
+		SetPortFlags(RoundtripPortFlags);
+	}
+
+	FORCENOINLINE void ReadProperties(UObject* Object, UObject* Defaults)
+	{
+		UClass* Class = Object->GetClass();
+		Class->SerializeTaggedProperties(*this, reinterpret_cast<uint8*>(Object), Class, reinterpret_cast<uint8*>(Defaults));
+	}
+
+	virtual FArchive& operator<<(FText& Value) override
+	{ 
+		int32 Idx = ReadValue<int32>();
+		Value = Idx == INDEX_NONE ? FText::GetEmpty() : Texts[Idx];
+		return *this;
+	}
+
+	virtual FArchive& operator<<(FName& Value) override
+	{
+		Value = FSensitiveName::FromUnstableInt(ReadValue<FSensitiveName::IntType>()).ToName();
+		return *this;
+	}
+
+	virtual FArchive& operator<<(UObject*& Value) override			{ reinterpret_cast<uint64&>(Value) = ReadValue<uint64>(); return *this; }
+	virtual FArchive& operator<<(FObjectPtr& Value) override		{ reinterpret_cast<uint64&>(Value) = ReadValue<uint64>(); return *this; }
+	virtual FArchive& operator<<(FWeakObjectPtr& Value) override	{ reinterpret_cast<uint64&>(Value) = ReadValue<uint64>(); return *this; }
+	virtual FArchive& operator<<(FLazyObjectPtr& Value) override	{ Value.ResetWeakPtr(); Value.GetUniqueID() = ReadValue<FUniqueObjectGuid>(); return *this; }
+	virtual FArchive& operator<<(FSoftObjectPtr& Value) override	{ Value.ResetWeakPtr(); Value.GetUniqueID().SerializePath(*this); return *this; }
+	virtual FArchive& operator<<(FSoftObjectPath& Value) override	{ Value.SerializePath(*this); return *this; }
+	virtual FString GetArchiveName() const override					{ return "FPropertyReader"; }
+
+	template<typename T>
+	T ReadValue()
+	{
+		T Out;
+		*this << Out;
+		return Out;
+	}
+};
+
+////////////////////////////////////////////////////////////////////////////////////////////////
+
+struct FInstance
+{
+	FString PathName;
+	UObject* Orig;
+	UObject* Arch;
+	UObject* Base;
+	UObject* PP;
+	UObject* TPS;
+	UObject* UPS;
+	FBindId Id;
+
+	void Init()
+	{
+		UClass* Class = Orig->GetClass();
+		Class->GetDefaultObject(/* create lazily */ true);
+		Arch = Orig->GetArchetype();
+		check(Arch != Orig);
+		Id = GUE.Names.IndexBindId(IndexType(Class));
+	}
+};
+
+static UObject* MakeEmptyInstance(UObject* Obj, FName Name)
+{
+	FStaticConstructObjectParameters Params(Obj->GetClass());
+	Params.Outer = Obj->GetOuter();//GetTransientPackage();
+	Params.Name = Name;
+	Params.SetFlags = Obj->GetFlags();
+	Params.Template = Obj->GetArchetype();
+	Params.bAssumeTemplateIsArchetype = true;
+	Params.bCopyTransientsFromClassDefaults = true;
+	return StaticConstructObject_Internal(Params);
+}
+
+static bool IncludeClass(UClass* Class)
+{
+	static const FName Exclusions[] = {
+		"CitySampleUnrealEdEngine", // Cloning MTAccessDetector crash
+		"GameFeaturePluginStateMachine", // Cloning ensure
+		"WorldSettings", // QAGame
+		// CitySample - Enum value 8 is undeclared in /Script/Engine.ERichCurveTangentMode, illegal value detected in /Script/Engine.RichCurveKey::TangentMode
+		"AnimSequence", "AnimationSequencerDataModel", "MovieSceneControlRigParameterSection",
+		// Enum value 7 is undeclared in /Script/Text3D.EText3DFontStyleFlags, illegal value detected in /Script/Text3DEditor.Text3DEditorFont::FontStyleFlags
+		// FontStyleFlags	Monospace | Bold | Italic (7 '\a')	UnrealEditor-Text3DEditor.dll!EText3DFontStyleFlags
+		"Text3DEditorFontSubsystem",
+	};
+	static const FName SuperExclusions[] = { "LevelScriptActor" };
+
+	if (!!Algo::Find(Exclusions, Class->GetFName()))
+	{
+		return false;
+	}
+
+	// Exclude IDOs
+	static constexpr EClassFlags IdoFlags = CLASS_NotPlaceable | CLASS_Hidden | CLASS_HideDropDown;
+	if (Class->HasAllClassFlags(IdoFlags))
+	{
+		return false;
+	}
+
+	for (UStruct* Super = Class->GetInheritanceSuper(); Super; Super = Super->GetInheritanceSuper())
+	{
+		if (!!Algo::Find(SuperExclusions, Super->GetFName()))
+		{
+			return false;
+		}
+	}
+	return true;
+}
+
+FORCENOINLINE static void SavePlainProps(FBatchSaver& Batch, TConstArrayView<FInstance> Instances)
+{
+	for (const FInstance& Instance : Instances)
+	{	
+		Batch.Save(Instance.Id, Instance.Orig, Instance.Base);
+	}
+}
+
+FORCENOINLINE static void LoadPlainProps(FMemoryBatchLoader& Batch, TConstArrayView<FInstance> Instances)
+{
+	for (const FInstance& Instance : Instances)
+	{	
+		Batch.Load(Instance.PP);
+	}
+}
+
+template<bool bUps>
+FORCENOINLINE void SaveArchive(FPropertyWriter& Archive, TConstArrayView<FInstance> Instances)
+{
+	Archive.SetUseUnversionedPropertySerialization(bUps);
+	for (const FInstance& Instance : Instances)
+	{
+		Archive.WriteProperties(Instance.Orig, Instance.Base);
+	}
+}
+
+template<bool bUps>
+FORCENOINLINE void LoadArchive(FPropertyReader& Archive, TConstArrayView<FInstance> Instances)
+{
+	Archive.SetUseUnversionedPropertySerialization(bUps);
+	for (const FInstance& Instance : Instances)
+	{
+		Archive.ReadProperties(bUps ? Instance.UPS : Instance.TPS, Instance.Base);
+	}
+}
+
+////////////////////////////////////////////////////////////////////////////////////////////////
+
+class FStableNameBatchIds final : public FStableBatchIds
+{
+	TConstArrayView<FSensitiveName> Names;
+public:
+	FStableNameBatchIds(FSchemaBatchId Batch, TConstArrayView<FSensitiveName> InNames) : FStableBatchIds(Batch), Names(InNames) {}
+	using FStableBatchIds::AppendString;
+
+	virtual uint32				NumNames() const override							{ return static_cast<uint32>(Names.Num()); }
+	virtual void				AppendString(FUtf8Builder& Out, FNameId Name) const override { Names[Name.Idx].AppendString(Out); }
+};
+
+[[nodiscard]] static FSchemaBatchId ParseBatch(
+	TArray64<uint8>& OutData,
+	TArray<FStructView>& OutObjects,
+	FUtf8StringView YamlView)
+{
+	// Parse yaml
+	ParseYamlBatch(OutData, YamlView);
+
+	// Grab and mount parsed schemas
+	FByteReader It(MakeMemoryView(OutData));
+	const uint32 SchemasSize = It.Grab<uint32>();
+	FMemoryView SchemasView = It.GrabSlice(SchemasSize);
+	const FSchemaBatch* Schemas = ValidateSchemas(SchemasView);
+	FSchemaBatchId Batch = MountReadSchemas(Schemas);
+
+	// Grab parsed objects
+	while (uint64 NumBytes = It.GrabVarIntU())
+	{	
+		FByteReader ObjIt(It.GrabSlice(NumBytes));
+		FStructSchemaId Schema = { ObjIt.Grab<uint32>() };
+		OutObjects.Add({ { Schema, Batch }, ObjIt });
+	}
+	
+	return Batch;
+}
+
+static void RoundtripText(
+	const FBatchIds& BatchIds,
+	TConstArrayView<FStructView> Objects,
+	TConstArrayView<FInstance> Instances,
+	ESchemaFormat Format,
+	FStringView AssetName = {})
+{
+	// Print yaml
+	UE_LOGFMT(LogPlainPropsUObject, Display, "Printing to PlainProps text using {Format}...", ToString(Format));
+	TUtf8StringBuilder<256> Yaml;
+	Yaml.Reserve(INT_MAX);
+	PrintYamlBatch(Yaml, BatchIds, Objects);
+	FUtf8StringView YamlView = Yaml.ToView();
+
+	// Write to file
+	if (AssetName.Len() == 0)
+	{
+		AssetName = (Format == ESchemaFormat::InMemoryNames ? TEXTVIEW("DiffTest_InMemoryNames") : TEXTVIEW("DiffTest_StableNames"));
+	}
+	const FString YamlPath = FPaths::Combine(FPaths::ProjectSavedDir(), TEXT("PlainProps"), AssetName) + TEXT(".yaml");
+	if (TUniquePtr<FArchive> FileWriter(IFileManager::Get().CreateFileWriter(*YamlPath)); FileWriter)
+	{
+		UE_LOGFMT(LogPlainPropsUObject, Display, "Writing {KB}KB yaml as {YamlPath}...", Yaml.Len() >> 10, *YamlPath);
+		FileWriter->Serialize((void*)YamlView.GetData(), YamlView.Len());
+	}
+
+	// Parse yaml
+	UE_LOGFMT(LogPlainPropsUObject, Display, "Parsing PlainProps text using {Format}...", ToString(Format));
+	TArray64<uint8> Data;
+	TArray<FStructView> ParsedObjects;
+	FSchemaBatchId ParsedBatch = ParseBatch(Data, ParsedObjects, YamlView);
+
+	if (Format == ESchemaFormat::StableNames && Instances.Num())
+	{
+		UE_LOGFMT(LogPlainPropsUObject, Display, "Diffing PlainProps parsed objects using {Format}...", ToString(Format));
+
+		// Diff schemas
+		check(!DiffSchemas(BatchIds.GetBatchId(), ParsedBatch));
+
+		// Diff objects
+		check(Objects.Num() == ParsedObjects.Num());
+		check(Objects.Num() == Instances.Num());
+		const int32 NumObjects = FMath::Min(Objects.Num(), ParsedObjects.Num());
+		uint32 NumDiffs = 0;
+		TUtf8StringBuilder<256> Diffs;
+		for (int32 I = 0; I < NumObjects; ++I)
+		{
+			FStructView In = Objects[I];
+			FStructView Parsed = ParsedObjects[I];
+			FReadDiffPath DiffPath;
+			if (DiffStruct(In, Parsed, DiffPath))
+			{
+				PrintDiff(Diffs, BatchIds, DiffPath);
+				Diffs.Append(" in ");
+				Diffs.Append(Instances[I].PathName);
+				Diffs.Append("\n");
+				++NumDiffs;
+			}
+		}
+		UE_LOGFMT(LogPlainPropsUObject, Display,
+			"Detected {Diffs} diffs in {Objs} PlainProps parsed objects from {KB}KB yaml text using StableNames\n{Diffs}",
+			NumDiffs, NumObjects, Yaml.Len() >> 10, Diffs.ToString());
+	}
+
+	// Unmount parsed schemas
+	UnmountReadSchemas(ParsedBatch);
+}
+
+class FBatchTextRoundtripper
+{
+public:
+	FBatchTextRoundtripper(FMemoryView Data, ESchemaFormat InFormat) : Format(InFormat)
+	{
+		FByteReader It(Data);
+		verify(It.Grab<uint32>() == Magics[0]);
+
+		// Read FNames when using Stable Names
+		if (Format == ESchemaFormat::StableNames)
+		{
+			ReadStableNames(It);
+		}
+		
+		// Read and mount schemas
+		verify(It.Grab<uint32>() == Magics[1]);
+		It.SkipAlignmentPadding<uint32>();
+		const uint32 SchemasSize = It.Grab<uint32>();
+		FMemoryView SavedSchemasView = It.GrabSlice(SchemasSize);
+		const FSchemaBatch* SavedSchemas = ValidateSchemas(SavedSchemasView);
+		verify(It.Grab<uint32>() == Magics[2]);
+		FSchemaBatchId Batch = MountReadSchemas(SavedSchemas);
+		
+		// Read objects
+		while (uint64 NumBytes = It.GrabVarIntU())
+		{	
+			FByteReader ObjIt(It.GrabSlice(NumBytes));
+			verify(ObjIt.Grab<uint32>() == Magics[3]);
+			FStructSchemaId Id = { ObjIt.Grab<uint32>() };
+			Objects.Add({ { Id, Batch }, ObjIt });
+		}
+		
+		verify(It.Grab<uint32>() == Magics[4]);
+		check(!Objects.IsEmpty());
+
+		// Create BatchIds 
+		if (Format == ESchemaFormat::StableNames)
+		{
+			BatchIds = MakeUnique<FStableNameBatchIds>(Batch, Names);
+		}
+		else
+		{
+			BatchIds = MakeUnique<FMemoryBatchIds>(Batch, GUE.Names);
+		}
+	}
+
+	void ReadStableNames(FByteReader& It)
+	{
+		TConstArrayView<uint8> NameBatch = GrabNumAndArray<uint8>(It);
+		TConstArrayView<int32> Numbers = GrabNumAndArray<int32>(It);
+		FMemoryReaderView Ar(NameBatch);
+		TArray<FDisplayNameEntryId> Entries = LoadNameBatch(Ar);
+		Names.Reserve(Entries.Num());
+		check(Entries.Num() == Numbers.Num());
+		for (int64 Idx = 0, Num = Entries.Num(); Idx < Num; ++Idx)
+		{
+			Names.Emplace(Entries[Idx].ToName(Numbers[Idx]));
+		}
+	}
+
+	~FBatchTextRoundtripper()
+	{
+		UnmountReadSchemas(BatchIds->GetBatchId());
+	}
+
+	void RoundtripTextForDiffTest(TConstArrayView<FInstance> Instances) const
+	{
+		PlainProps::UE::RoundtripText(*BatchIds, Objects, Instances, Format, {});
+	}
+
+	void RoundtripTextForPackage(FStringView PackageName) const
+	{
+		PlainProps::UE::RoundtripText(*BatchIds, Objects, {}, Format, PackageName);
+	}
+
+private:
+	TArray<FStructView>			Objects;
+	TArray<FSensitiveName>		Names;
+	TUniquePtr<FBatchIds>		BatchIds;
+	ESchemaFormat				Format;
+	
+	template<typename T>
+	static TConstArrayView<T> GrabNumAndArray(/* in-out */ FByteReader& It)
+	{
+		uint32 Num = It.Grab<uint32>();
+		return MakeArrayView(reinterpret_cast<const T*>(It.GrabBytes(Num * sizeof(T))), Num);
+	}
+};
+
+////////////////////////////////////////////////////////////////////////////////
+
+struct FDiffDebug
+{
+	TArray<FInstance*> Instances;
+	volatile const UTF8CHAR* Str;
+};
+static FDiffDebug GPpDiff, GTpsDiff, GUpsDiff;
+
+static void DumpUObjectToJsonFile(const UObject* Object)
+{
+	FString JsonString;
+	if (FJsonObjectConverter::UStructToJsonObjectString(Object->GetClass(), Object, JsonString))
+	{
+		TStringBuilder<256> JsonPath;
+		JsonPath << FPaths::Combine(FPaths::ProjectSavedDir(), TEXT("PlainProps"), Object->GetPackage()->GetName());
+		JsonPath << "." << Object->GetName() + ".json";
+
+		UE_LOGFMT(LogPlainPropsUObject, Display, "Writing {JsonPath}...", JsonPath);
+		FFileHelper::SaveStringToFile(JsonString, JsonPath.ToString(), FFileHelper::EEncodingOptions::ForceUTF8WithoutBOM);
+	}
+}
+
+////////////////////////////////////////////////////////////////////////////////
+
+#if WITH_EDITOR
+
+struct FLinkerPropertyBatch
+{
+	explicit FLinkerPropertyBatch(const FLinkerSave& InLinkerSave)
+		: LinkerSave(&InLinkerSave)
+	{
+	}
+
+	explicit FLinkerPropertyBatch(FLinkerLoad& InLinkerLoad)
+		: LinkerLoad(&InLinkerLoad)
+	{
+	}
+
+	static FSoleMemberSpec Spec(FName*, FDeclId Id)
+	{
+		return {Id, GUE.Members.Id, Specify<uint64>()};
+	}
+
+	void Save(FMemberBuilder& Out, FName In, const FSaveContext&)
+	{
+		// see FLinkerSave::MapName
+		int32 Index = LinkerSave->NameIndices.FindRef(In.GetDisplayIndex(), INDEX_NONE);
+		int32 Number = In.GetNumber();
+		uint64 Id;
+		if (Index == INDEX_NONE)
+		{
+			UE_LOGFMT(LogPlainPropsUObject, Warning, "FName {Name} is not mapped.", In);
+			Id = 6666666;
+		}
+		else
+		{
+			Id = Index != INDEX_NONE ? ((uint64(Number) << 32) | Index) : 6666666;
+		}
+		Out.Add(GUE.Members.Id, Id);
+	}
+
+	void Load(FName& Out, FStructLoadView In)
+	{
+		uint64 Id;
+		LoadSole<uint64>(&Id, In);
+		if (Id != 6666666)
+		{
+			int32 Index = Id & 0xFFFFFFFF;
+			int32 Number = Id >> 32;
+			FNameEntryId Entry = LinkerLoad->NameMap[Index];
+			FName Name = FName::CreateFromDisplayId(Entry, Number);
+			Out = Name;
+		}
+		else
+		{
+			Out = FName();
+		}
+	}
+
+	static FSoleMemberSpec Spec(FText*, FDeclId Id)
+	{
+		return {Id, GUE.Members.Id, Specify<int32>()};
+	}
+
+	void Save(FMemberBuilder& Out, const FText& In, const FSaveContext&)
+	{
+		Out.Add(GUE.Members.Id, Texts.Num());
+		Texts.Add(In);
+	}
+
+	void Load(FText& Out, FStructLoadView In) const
+	{
+		int32 Dummy = LoadSole<int32>(In);
+		Out = {};//TODO: Texts[LoadSole<int32>(In)];
+	}
+
+	static FSoleMemberSpec Spec(FObjectHandle*, FDeclId Id)
+	{
+		return {Id, GUE.Members.Id, Specify<int32>()};
+	}
+
+	void Save(FMemberBuilder& Out, FObjectHandle In, const FSaveContext&)
+	{
+		// see FLinkerSave::MapObject
+		UObject* Object = FObjectPtr(In).Get();
+		FPackageIndex Index = LinkerSave->ObjectIndicesMap.FindRef(Object);
+		int32 Idx = *reinterpret_cast<int32*>(&Index);
+		if (Object && Index.IsNull())
+		{
+			UE_LOGFMT(LogPlainPropsUObject, Warning, "UObject {Object} is not mapped.", Object->GetFullName());
+			Idx = 9999999;
+		}
+		Out.Add(GUE.Members.Id, Idx);
+	}
+
+	void Load(FObjectHandle& Out, FStructLoadView In)
+	{
+		int32 Idx;
+		LoadSole<int32>(&Idx, In);
+		if (Idx != 9999999)
+		{
+			FPackageIndex PackageIdx = *reinterpret_cast<FPackageIndex*>(&Idx);
+			if (PackageIdx.IsImport())
+			{
+				UObject* Object = LinkerLoad->ImportMap[PackageIdx.ToImport()].XObject;
+				if (Object == nullptr)
+				{
+#if PP_WITH_ENGINE_HOOKS
+					// Object = LinkerLoad->ResolveResource(PackageIdx);
+#endif
+				}
+				Out = ::UE::CoreUObject::Private::MakeObjectHandle(Object);
+			}
+			else if (PackageIdx.IsExport())
+			{
+				UObject* Object = LinkerLoad->ExportMap[PackageIdx.ToExport()].Object;
+				if (Object == nullptr)
+				{
+#if PP_WITH_ENGINE_HOOKS
+					// Object = LinkerLoad->ResolveResource(PackageIdx);
+#endif
+				}
+				Out = ::UE::CoreUObject::Private::MakeObjectHandle(Object);
+			}
+			else
+			{
+				check(IsObjectHandleNull(Out));
+			}
+		}
+		else
+		{
+			Out = FObjectHandle();
+		}
+	}
+
+	static FSoleMemberSpec Spec(FWeakObjectPtr*, FDeclId Id)
+	{
+		return {Id, GUE.Members.Id, Specify<uint64>()};
+	}
+
+	void Save(FMemberBuilder& Out, const FWeakObjectPtr& In, const FSaveContext& Ctx)
+	{
+		// Save ObjectSerialNumber + ObjectIndex a single uint64
+		UObject* Object = In.Get(true);
+		Save(Out, ::UE::CoreUObject::Private::MakeObjectHandle(Object), Ctx);
+	}
+
+	void Load(FWeakObjectPtr& Out, FStructLoadView In)
+	{
+		FObjectHandle Handle;
+		Load(Handle, In);
+		Out = ::UE::CoreUObject::Private::ResolveObjectHandle(Handle);
+	}
+
+	static FSoleMemberSpec Spec(FSoftObjectPtr*, FDeclId Id)
+	{
+		return {Id, GUE.Members.Id, FMemberSpec(GUE.Structs.SoftObjectPath)};
+	}
+
+	static void Save(FMemberBuilder& Out, const FSoftObjectPtr& In, const FSaveContext& Ctx)
+	{
+		FBuiltStruct* SoftPath = SaveStruct(&In.GetUniqueID(), GUE.Structs.SoftObjectPath, Ctx);
+		Out.AddStruct(GUE.Members.Id, GUE.Structs.SoftObjectPath, SoftPath);
+	}
+
+	static void Load(FSoftObjectPtr& Out, FStructLoadView In)
+	{
+		Out.ResetWeakPtr();
+		LoadSoleStruct(&Out.GetUniqueID(), In);
+	}
+
+	static FSoleMemberSpec Spec(FLazyObjectPtr*, FDeclId Id)
+	{
+		return {Id, GUE.Members.Id, FMemberSpec(GUE.Structs.Guid)};
+	}
+	
+	static void Save(FMemberBuilder& Out, const FLazyObjectPtr& In, const FSaveContext& Ctx)
+	{
+		FBuiltStruct* Guid = SaveStruct(&In.GetUniqueID(), GUE.Structs.Guid, Ctx);
+		Out.AddStruct(GUE.Members.Id, GUE.Structs.Guid, Guid);
+	}
+
+	static void Load(FLazyObjectPtr& Out, FStructLoadView In)
+	{
+		Out.ResetWeakPtr();
+		LoadSoleStruct(&Out.GetUniqueID(), In);
+	}
+
+private:
+	union
+	{
+		const FLinkerSave* LinkerSave;
+		FLinkerLoad* LinkerLoad;
+		const FLinker* Linker;
+	};
+
+	TArray<FText> Texts; // Tricky to serialize non-intrusively
+};
+
+struct FPlainPropsLinkerSandbox
+{
+	static constexpr uint32 LoadFlags = (LOAD_Async | LOAD_NoVerify | LOAD_SkipLoadImportedPackages);
+	static constexpr uint32 CloneNumber = 8848;
+
+	struct FSavedPackage
+	{
+		FName PackageName;
+		FPackagePath PackagePath;
+		FSharedBuffer PackageData;
+		TArray64<uint8> PlainPropsData;
+	};
+
+	struct FLinkerBatch
+	{
+		FLinkerBatch(const FSavedPackage& Saved)
+			: PackagePath(Saved.PackagePath)
+			, PackageName(Saved.PackageName.GetNumber() ? FName(*Saved.PackageName.ToString(), CloneNumber) : FName(Saved.PackageName, CloneNumber))
+			, Package(NewObject<UPackage>(/*Outer*/nullptr, PackageName))
+			, LinkerReader(static_cast<const uint8*>(Saved.PackageData.GetData()), Saved.PackageData.GetSize(), ELargeMemoryReaderFlags::Persistent, PackageName)
+			, LinkerLoad(FLinkerLoad::CreateLinker(FUObjectThreadContext::Get().GetSerializeContext(), Package, PackagePath, LoadFlags, &LinkerReader))
+			, Properties(*LinkerLoad)
+			, Customs(Properties, GUE.Customs)
+			, BatchView(MakeMemoryView(Saved.PlainPropsData))
+			, BatchLoader(Customs.Overlay, BatchView)
+		{
+		}
+
+		FPackagePath PackagePath;
+		FName PackageName;
+		UPackage* Package;
+		FLargeMemoryReader LinkerReader;
+		FLinkerLoad* LinkerLoad;
+		FLinkerPropertyBatch Properties;
+		TCustomPropertyBindings<FLinkerPropertyBatch> Customs;
+		FMemoryView BatchView;
+		FPersistentBatchLoader BatchLoader;
+	};
+
+	static TMap<FName, FSavedPackage> SavedPackages;
+	static TMap<FName, TUniquePtr<FLinkerBatch>> LoadBatches;
+
+	static void SavePlainProps(const FLinkerSave& LinkerSave)
+	{
+		FName PackageName = LinkerSave.LinkerRoot->GetFName();
+
+		UE_LOGFMT(LogPlainPropsUObject, Log, " Saving PlainProps for {PackageName}...", PackageName);
+		FLinkerPropertyBatch Properties(LinkerSave);
+		TCustomPropertyBindings<FLinkerPropertyBatch> Customs(Properties, GUE.Customs);
+		FBatchSaver BatchSaver(Customs.Overlay, LinkerSave.ExportMap.Num());
+		for (const FObjectExport& Export : LinkerSave.ExportMap)
+		{
+			const UStruct* Struct = SkipEmptyBases(Export.Object->GetClass());
+			FBindId Id = GUE.Names.IndexBindId(PlainProps::UE::IndexType(Struct));
+			if (GUE.Schemas.FindStruct(Id))
+			{
+				BatchSaver.Save(Id, Export.Object, Export.Object->GetArchetype());
+			}
+			else
+			{
+				BatchSaver.Skip(Id, Export.Object);
+			}
+		}
+		TArray64<uint8> StableData = BatchSaver.Write();
+
+		FSavedPackage& SavedPackage = SavedPackages.FindOrAdd(PackageName);
+		SavedPackage.PackageName = PackageName;
+		SavedPackage.PackagePath = LinkerSave.LinkerRoot->GetLoadedPath();
+		SavedPackage.PlainPropsData = MoveTemp(StableData);
+	}
+
+	static void SavePackageToMemory(const FLinkerSave& LinkerSave, uint8* Data, int64 Size)
+	{
+		FName PackageName = LinkerSave.LinkerRoot->GetFName();
+		FSavedPackage& SavedPackage = SavedPackages.FindChecked(PackageName);
+		SavedPackage.PackageData = FSharedBuffer::TakeOwnership(Data, Size, FMemory::Free);
+	}
+
+	static bool OpenPackageReader(const FPackagePath& PackagePath, FArchive*& Ar)
+	{
+		FName PackageName = PackagePath.GetPackageFName();
+		if (FSavedPackage* SavedData = SavedPackages.Find(PackageName))
+		{
+			Ar = new FLargeMemoryReader(static_cast<const uint8*>(SavedData->PackageData.GetData()), SavedData->PackageData.GetSize(), ELargeMemoryReaderFlags::Persistent, PackageName);
+			return true;
+		}
+		return false;
+	}
+
+	static void MountPackageFromMemory(FName PackageName)
+	{
+		if (FSavedPackage* SavedData = SavedPackages.Find(PackageName))
+		{
+			FLinkerBatch* LinkerBatch = new FLinkerBatch(*SavedData);
+			LoadBatches.Emplace(PackageName, LinkerBatch);
+		}
+	}
+
+	static void LoadPackageFromMemory(FName PackageName)
+	{
+		if (TUniquePtr<FLinkerBatch>* LinkerBatch = LoadBatches.Find(PackageName))
+		{
+			UE_LOGFMT(LogPlainPropsUObject, Log, " Loading PlainProps for {Package}...", PackageName);
+			LoadPackage((*LinkerBatch)->Package, (*LinkerBatch)->PackagePath, LoadFlags);
+		}
+	}
+
+	static bool LoadObject(FStructuredArchive::FSlot& Slot, uint8* Data)
+	{
+		UObject* Object = (UObject*)Data;
+		FName PackageName = Object->GetPackage()->GetFName();
+		if (PackageName.GetNumber() == CloneNumber)
+		{
+			PackageName.SetNumber(0);
+		}
+
+		TUniquePtr<FLinkerBatch>* LinkerBatch = LoadBatches.Find(PackageName);
+		if (!LinkerBatch)
+		{
+			return false;
+		}
+
+		FPersistentBatchLoader& BatchLoader = (*LinkerBatch)->BatchLoader;
+		FLinkerLoad* LinkerLoad = (*LinkerBatch)->LinkerLoad;
+		check(LinkerLoad == Object->GetLinker());
+
+		FUObjectSerializeContext* SerializeContext = FUObjectThreadContext::Get().GetSerializeContext();
+		check(Object == SerializeContext->SerializedObject);
+
+		const int32 ExportIndex = Object->GetLinkerIndex();
+		const FObjectExport& Export = LinkerLoad->ExportMap[ExportIndex];
+		check(Export.Object == Object);
+
+		BatchLoader.LoadObject(Object, ExportIndex);
+
+		const int32 Pos = Slot.GetUnderlyingArchive().Tell();
+		check(Pos == SerializeContext->SerializedObjectScriptStartOffset);
+		const int32 ScriptSerializationSize = Export.ScriptSerializationEndOffset - Export.ScriptSerializationStartOffset;
+		check(ScriptSerializationSize >= 0);
+		Slot.GetUnderlyingArchive().Seek(Pos + ScriptSerializationSize);
+		return true;
+	}
+};
+
+TMap<FName, FPlainPropsLinkerSandbox::FSavedPackage> FPlainPropsLinkerSandbox::SavedPackages;
+TMap<FName, TUniquePtr<FPlainPropsLinkerSandbox::FLinkerBatch>> FPlainPropsLinkerSandbox::LoadBatches;
+
+#endif
+
+////////////////////////////////////////////////////////////////////////////////
+
+int32 RoundtripViaBatch(TConstArrayView<UObject*> Objects, ERoundtrip Options)
+{
+	TArray<FInstance> Instances;
+	if (Objects.Num())
+	{
+		UE_LOGFMT(LogPlainPropsUObject, Display, "Gathering UObjects in Assets...");
+		for (const UObject* Object : Objects)
+		{
+			ForEachObjectWithPackage(Object->GetPackage(), [&Instances](UObject* Obj)
+			{
+				UClass* Class = Obj->GetClass();
+				if (ShouldBind(Class) && IncludeClass(Class))
+				{
+					Instances.Emplace(Obj->GetPathName(), Obj);
+				}
+				return true;
+			});
+		}
+	}
+	else
+	{
+		UE_LOGFMT(LogPlainPropsUObject, Display, "Gathering all non-empty UObjects...");
+		static constexpr EObjectFlags SkipFlags = RF_ClassDefaultObject | RF_MirroredGarbage | RF_InheritableComponentTemplate;
+		for (TObjectIterator<UObject> It(SkipFlags); It; ++It)
+		{
+			UObject* Object = *It;
+			UClass* Class = Object->GetClass();
+			if (ShouldBind(Class) && IncludeClass(Class))
+			{
+				Instances.Emplace(Object->GetPathName(), Object);
+			}
+		}
+	}
+
+	UE_LOGFMT(LogPlainPropsUObject, Display, "Sorting {Num} UObjects ...", Instances.Num());
+	Algo::Sort(Instances, [](const FInstance& A, const FInstance& B)
+	{
+		return FPlatformString::Strcmp(*A.PathName, *B.PathName) < 0;
+	});
+
+	// Create CDOs if needed and then clones for PP and TPS tests
+	UE_LOGFMT(LogPlainPropsUObject, Display, "Cloning {Num} UObjects up to 4 times...", Instances.Num());
+	for (FInstance& Instance : Instances)
+	{
+		Instance.Init();
+	}
+	FlushAsyncLoading();
+
+	for (FInstance& Instance : Instances)
+	{
+		uint32 N = 1 + &Instance - Instances.GetData();
+		Instance.Base = Instance.Arch ? MakeEmptyInstance(Instance.Orig, FName("Base", N)) : nullptr;
+		if (EnumHasAnyFlags(Options, ERoundtrip::PP))
+		{
+			Instance.PP = MakeEmptyInstance(Instance.Orig, FName("PP", N));
+		}
+		if (EnumHasAnyFlags(Options, ERoundtrip::TPS))
+		{
+			Instance.TPS = MakeEmptyInstance(Instance.Orig, FName("TPS", N));
+		}
+		if (EnumHasAnyFlags(Options, ERoundtrip::UPS))
+		{
+			Instance.UPS = MakeEmptyInstance(Instance.Orig, FName("UPS", N));
+		}
+	}
+
+	// Save
+	UE_LOGFMT(LogPlainPropsUObject, Display, "Saving UObjects to PlainProps with InMemoryNames...");
+	FMemoryBatch Plain;
+	TCustomPropertyBindings<FMemoryPropertyBatch> Customs(Plain.Properties, GUE.Customs);
+	{
+		FBatchSaver Batch(Customs.Overlay, GUObjectArray.GetObjectArrayNum());
+		SavePlainProps(Batch, Instances);
+		Plain.Data = Batch.Write(/* out */ &Plain.RuntimeIds);
+
+		if (EnumHasAnyFlags(Options, ERoundtrip::TextMemory))
+		{
+			FBatchTextRoundtripper MemoryBatch(MakeMemoryView(Plain.Data), ESchemaFormat::InMemoryNames);
+			MemoryBatch.RoundtripTextForDiffTest(Instances);
+		}
+		if (EnumHasAnyFlags(Options, ERoundtrip::TextStable))
+		{
+			UE_LOGFMT(LogPlainPropsUObject, Display, "Saving UObjects to PlainProps with StableNames...");
+			TArray64<uint8> StableData = Batch.Write();
+
+			FBatchTextRoundtripper StableBatch(MakeMemoryView(StableData), ESchemaFormat::StableNames);
+			StableBatch.RoundtripTextForDiffTest(Instances);
+		}
+	}
+
+	// Load
+	uint32 NumPpDiffs = 0;
+	if (EnumHasAnyFlags(Options, ERoundtrip::PP))
+	{
+		UE_LOGFMT(LogPlainPropsUObject, Display, "Loading UObjects from PlainProps...");
+		FMemoryBatchLoader Batch(Customs.Overlay, MakeMemoryView(Plain.Data), Plain.RuntimeIds);
+		LoadPlainProps(Batch, Instances);
+
+		// Diff original vs PlainProps
+		UE_LOGFMT(LogPlainPropsUObject, Display, "Diffing UObjects roundtripped via PlainProps...");
+		FDiffContext DiffCtx = {{ GUE.Schemas, Customs.Overlay }};
+		TUtf8StringBuilder<256> PpDiffPaths;
+		for (FInstance& Instance : Instances)
+		{
+			if (DiffStructs(Instance.Orig, Instance.PP, Instance.Id, /* in-out */ DiffCtx))
+			{
+				PpDiffPaths.Append("\n");
+				GPpDiff.Instances.Add(&Instance);
+				PrintDiff(/* out */ PpDiffPaths, GUE.Names, DiffCtx.Out);
+
+				//Batch.Reload(Instance.UPS, &Instance - Instances.GetData());
+				DiffCtx.Out.Reset();
+				PpDiffPaths.Append(" in ");
+				PpDiffPaths.Append(Instance.Orig->GetFullName());
+				++NumPpDiffs;
+			}
+		}
+		GPpDiff.Str = PpDiffPaths.GetData();
+		UE_LOGFMT(LogPlainPropsUObject, Display, "Detected {Diffs} diffs in {Objs} UObjects saved in a {KB}KB value stream using PlainProps{DiffText}", NumPpDiffs, Instances.Num(), Plain.Data.NumBytes() / 1024, PpDiffPaths.ToString());
+		for (const FInstance* Instance : GPpDiff.Instances)
+		{
+			DumpUObjectToJsonFile(Instance->Orig);
+			DumpUObjectToJsonFile(Instance->PP);
+		}
+	}
+
+	FArchivedProperties Tps;
+	if (EnumHasAnyFlags(Options, ERoundtrip::TPS))
+	{
+		UE_LOGFMT(LogPlainPropsUObject, Display, "Saving UObjects to TPS archive...");
+		FPropertyWriter Archive(/* out */ Tps);
+		SaveArchive<false>(Archive, Instances);
+	}
+
+	if (EnumHasAnyFlags(Options, ERoundtrip::TPS))
+	{
+		UE_LOGFMT(LogPlainPropsUObject, Display, "Loading UObjects from TPS archive...");
+		FPropertyReader Archive(Tps);
+		LoadArchive<false>(Archive, Instances);
+	}
+
+	static const FName SkipClasses[] = {
+		"BodySetup", // Skips some structs due to native FCollisionResponse::operator==
+		"NiagaraScript", "NiagaraNodeFunctionCall", "NiagaraMeshRendererProperties", // FNiagaraTypeDefinition::Serialize resets ClassStructOrEnum
+	};
+
+	if (EnumHasAnyFlags(Options, ERoundtrip::TPS))
+	{
+		// Diff original vs TPS
+		UE_LOGFMT(LogPlainPropsUObject, Display, "Diffing UObjects roundtripped via TPS...");
+		FDiffContext DiffCtx = {{ GUE.Schemas, Customs.Overlay }};
+		TUtf8StringBuilder<256> TpsDiffs;
+		TArray<int32> TpsDiffIdxs;
+		for (FInstance& Instance : Instances)
+		{
+			if (Algo::Find(SkipClasses, Instance.Orig->GetClass()->GetFName()))
+			{
+				continue;
+			}
+
+			if (DiffStructs(Instance.Orig, Instance.TPS, Instance.Id, /* in-out */ DiffCtx))
+			{
+				GTpsDiff.Instances.Add(&Instance);
+				PrintDiff(/* out */ TpsDiffs, GUE.Names, DiffCtx.Out);
+				
+				DiffCtx.Out.Reset();
+				TpsDiffs.Append(" in ");
+				TpsDiffs.Append(Instance.Orig->GetFullName());
+				TpsDiffs.Append("\n");	
+				TpsDiffIdxs.Add(&Instance - &Instances[0]);
+				
+				//FArchivedProperties Tmp;
+				//FPropertyWriter(/* out */ Tmp).WriteProperties(Instance.Orig, Instance.Base);
+				//FPropertyReader(Tmp).ReadProperties(Instance.UPS, Instance.Base);
+
+				//FDiffContext TmpCtx = DiffCtx;
+				//bool bOU = DiffStructs(Instance.Id, Instance.Orig, Instance.UPS, /* in-out */ TmpCtx);
+				//bool bOT = DiffStructs(Instance.Id, Instance.Orig, Instance.TPS, /* in-out */ TmpCtx);
+				//bool bOP = DiffStructs(Instance.Id, Instance.Orig, Instance.PP, /* in-out */ TmpCtx);
+				//TUtf8StringBuilder<256> TmpDiff;
+				//PrintDiff(/* out */ TmpDiff, TmpCtx.Out, GUE.Names);
+			}
+		}
+		GTpsDiff.Str = TpsDiffs.GetData();
+		UE_LOGFMT(LogPlainPropsUObject, Display, "Detected {Diffs} diffs in {Objs} UObjects saved in a {KB}KB value stream using TPS", TpsDiffIdxs.Num(), Instances.Num(), Tps.Data.NumBytes() / 1024);
+	}
+
+
+	if (EnumHasAnyFlags(Options, ERoundtrip::UPS))
+	{
+		FArchivedProperties Ups;
+		UE_LOGFMT(LogPlainPropsUObject, Display, "Saving UObjects to UPS archive...");
+		{
+			FPropertyWriter Archive(/* out */ Ups);
+			SaveArchive<true>(Archive, Instances);
+		}
+
+		UE_LOGFMT(LogPlainPropsUObject, Display, "Loading UObjects from UPS archive...");
+		{
+			FPropertyReader Archive(Ups);
+			LoadArchive<true>(Archive, Instances);
+		}
+
+		// Diff original vs UPS
+		UE_LOGFMT(LogPlainPropsUObject, Display, "Diffing UObjects roundtripped via UPS...");
+		FDiffContext DiffCtx = {{ GUE.Schemas, Customs.Overlay }};
+		TUtf8StringBuilder<256> UpsDiffs;
+		TArray<int32> UpsDiffIdxs;
+		for (FInstance& Instance : Instances)
+		{
+			if (Algo::Find(SkipClasses, Instance.Orig->GetClass()->GetFName()))
+			{
+				continue;
+			}
+
+			if (DiffStructs(Instance.Orig, Instance.UPS, Instance.Id, /* in-out */ DiffCtx))
+			{
+				GUpsDiff.Instances.Add(&Instance);
+				PrintDiff(/* out */ UpsDiffs, GUE.Names, DiffCtx.Out);
+				
+				DiffCtx.Out.Reset();
+				UpsDiffs.Append(" in ");
+				UpsDiffs.Append(Instance.Orig->GetFullName());
+				UpsDiffs.Append("\n");	
+				UpsDiffIdxs.Add(&Instance - &Instances[0]);
+			}
+		}
+		GUpsDiff.Str = UpsDiffs.GetData();
+		UE_LOGFMT(LogPlainPropsUObject, Display, "Detected {Diffs} diffs in {Objs} UObjects saved in a {KB}KB value stream using UPS", UpsDiffIdxs.Num(), Instances.Num(), Ups.Data.NumBytes() / 1024);
+	}
+
+	return NumPpDiffs;
+}
+
+////////////////////////////////////////////////////////////////////////////////
+
+int32 RoundtripViaPackages(TConstArrayView<UObject*> Objects, ERoundtrip Options)
+{
+	TArray<UPackage*> Packages;
+	TArray<UObject*> AllPackages;
+
+	// Collect all packages
+	if (Objects.Num())
+	{
+		for (UObject* Object : Objects)
+		{
+			AllPackages.Add(Object->GetPackage());
+		}
+	}
+	else
+	{
+		GetObjectsOfClass(UPackage::StaticClass(), AllPackages);
+	}
+
+	// Filter out in-memory and BP packages
+	for (UObject* PackageObj : AllPackages)
+	{
+		if (!PackageObj->HasAnyFlags(RF_WasLoaded))
+		{
+			continue;
+		}
+		UPackage* Package = Cast<UPackage>(PackageObj);
+		if (Package->HasAnyPackageFlags(PKG_CompiledIn))
+		{
+			continue;
+		}
+		bool bExclude = false;
+		ForEachObjectWithPackage(Package, [&bExclude](UObject* Object)
+		{
+			UClass* Class = Object->GetClass();
+			if (Class->HasAnyClassFlags(CLASS_CompiledFromBlueprint))
+			{
+				bExclude = true;
+				return false;
+			}
+			if (!IncludeClass(Class))
+			{
+				UE_LOGFMT(LogPlainPropsUObject, Log, " Excluding PlainProps for {Package}", Object->GetPackage()->GetName());
+				bExclude = true;
+				return false;
+			}
+			return true;
+		});
+		if (bExclude)
+		{
+			continue;
+		}
+		Packages.Add(Package);
+	}
+
+#if PP_WITH_ENGINE_HOOKS
+	// Todo consider TGuardValue
+	FSavePackageSettings& SavePackageSettings = FSavePackageSettings::GetDefaultSettings();
+	SavePackageSettings.ExternalWriteExportsFunc = &FPlainPropsLinkerSandbox::SavePlainProps;
+	SavePackageSettings.ExternalWriteFileFunc = &FPlainPropsLinkerSandbox::SavePackageToMemory;
+	GSerializePropertiesOverride = &FPlainPropsLinkerSandbox::LoadObject;
+	GLoadFromPackageFileOverride = &FPlainPropsLinkerSandbox::OpenPackageReader;
+#endif
+
+	UE_LOGFMT(LogPlainPropsUObject, Display, "Saving {SaveNum} packages...", Packages.Num());
+	for (UPackage* Package : Packages)
+	{
+		UObject* Asset = Package->FindAssetInPackage();
+		FSavePackageArgs SaveArgs;
+		SaveArgs.TopLevelFlags = RF_Standalone;
+		SaveArgs.SaveFlags = SAVE_Async; // Write progress to temp memory instead of a temp file
+		const FString& Extension = Package->ContainsMap() ? FPackageName::GetMapPackageExtension() : FPackageName::GetAssetPackageExtension();
+		FString TempFilename = FPaths::Combine(*FPaths::ProjectSavedDir(), TEXT("PlainProps"), Package->GetName() + Extension);
+		// FString ActualFilename = FPackageName::LongPackageNameToFilename(Asset->GetPackage()->GetName(), Extension);
+
+		UPackage::Save(Package, Asset, *TempFilename, SaveArgs);
+	}
+
+#if WITH_EDITOR
+	FPlainPropsLinkerSandbox::SavedPackages.Reserve(Packages.Num());
+	FPlainPropsLinkerSandbox::LoadBatches.Reserve(Packages.Num());
+
+	UE_LOGFMT(LogPlainPropsUObject, Display, "Mounting {SaveNum} packages...", Packages.Num());
+	for (UPackage* Package : Packages)
+	{
+		FPlainPropsLinkerSandbox::MountPackageFromMemory(Package->GetFName());
+	}
+
+	UE_LOGFMT(LogPlainPropsUObject, Display, "Loading {SaveNum} packages...", Packages.Num());
+	for (UPackage* Package : Packages)
+	{
+		FPlainPropsLinkerSandbox::LoadPackageFromMemory(Package->GetFName());
+	}
+#endif
+
+#if PP_WITH_ENGINE_HOOKS
+	SavePackageSettings.ExternalWriteExportsFunc = nullptr;
+	GSerializePropertiesOverride = nullptr;
+#endif
+
+	return 0;
+}
+
+} // namespace PlainProps::UE
diff --git a/Engine/Plugins/Experimental/PlainPropsUObject/Source/Private/PlainPropsUObjectBindings.cpp b/Engine/Plugins/Experimental/PlainPropsUObject/Source/Private/PlainPropsUObjectBindings.cpp
--- a/Engine/Plugins/Experimental/PlainPropsUObject/Source/Private/PlainPropsUObjectBindings.cpp
+++ b/Engine/Plugins/Experimental/PlainPropsUObject/Source/Private/PlainPropsUObjectBindings.cpp
@@ -1,59 +1,35 @@
 // Copyright Epic Games, Inc. All Rights Reserved.
 
+#include "PlainPropsUObjectBindingsInternal.h"
 #include "PlainPropsRoundtripTest.h"
 #include "PlainPropsUObjectRuntime.h"
 #include "PlainPropsUObjectModule.h"
-#include "PlainPropsBuildSchema.h"
-#include "PlainPropsDiff.h"
-#include "PlainPropsParse.h"
-#include "PlainPropsPrint.h"
-#include "PlainPropsRead.h"
 #include "PlainPropsUeCoreBindings.h"
 #include "PlainPropsVisualize.h"
-#include "PlainPropsWrite.h"
 #include "Algo/Compare.h"
 #include "Algo/Find.h"
 #include "Containers/PagedArray.h"
-#include "HAL/FileManager.h"
 #include "Hash/xxhash.h"
-#include "JsonObjectConverter.h"
 #include "Logging/StructuredLog.h"
 #include "Math/PreciseFP.h"
-#include "Misc/AsciiSet.h"
-#include "Misc/CommandLine.h"
-#include "Misc/CoreDelegates.h"
 #include "Misc/DefinePrivateMemberPtr.h"
-#include "Misc/Filehelper.h"
-#include "Misc/Paths.h"
-#include "Misc/PackageName.h"
-#include "Serialization/Archive.h"
-#include "Serialization/BufferWriter.h"
-#include "Serialization/LargeMemoryReader.h"
-#include "Serialization/MemoryReader.h"
-#include "Serialization/MemoryWriter.h"
-#include "Serialization/UnversionedPropertySerializationInternal.h"
 #include "StructUtils/UserDefinedStruct.h"
 #include "Templates/UniquePtr.h"
 #include "UObject/AnsiStrProperty.h"
 #include "UObject/Class.h"
 #include "UObject/EnumProperty.h"
 #include "UObject/FieldPathProperty.h"
-#include "UObject/LinkerLoad.h"
-#include "UObject/LinkerSave.h"
-#include "UObject/NameBatchSerialization.h"
 #include "UObject/Object.h"
 #include "UObject/ObjectHandle.h"
 #include "UObject/ObjectResource.h"
 #include "UObject/PropertyOptional.h"
-#include "UObject/SavePackage.h"
 #include "UObject/StrProperty.h"
-#include "UObject/TextProperty.h"
-#include "UObject/Utf8StrProperty.h"
-#include "UObject/UnrealType.h"
 #include "UObject/UObjectIterator.h"
 #include "UObject/UObjectThreadContext.h"
-#include "UObject/VerseValueProperty.h"
+#include "UObject/UnrealType.h"
+#include "UObject/Utf8StrProperty.h"
 #include "UObject/VerseStringProperty.h"
+#include "UObject/VerseValueProperty.h"
 
 
 #define PP_WITH_ENGINE_HOOKS 0
@@ -99,7 +75,7 @@ struct FDelegateAccess
 #endif // UE_DETECT_DELEGATES_RACE_CONDITIONS
 
 
-DEFINE_LOG_CATEGORY_STATIC(LogPlainPropsUObject, Log, All);
+DEFINE_LOG_CATEGORY(LogPlainPropsUObject);
 
 namespace PlainProps::UE
 {
@@ -397,7 +373,7 @@ static bool HasAny(uint64 Flags)
 
 ////////////////////////////////////////////////////////////////////////////////////////////////
 
-static FType IndexType(const UField* Field)
+FType IndexType(const UField* Field)
 {
 	check(Field);
 	FTypenameId Name = GUE.Names.MakeTypename(Field->GetFName());
@@ -3028,7 +3004,7 @@ struct FStaticArrayBinding : IItemRangeBinding
 //		UObjects are never instantiated during serialization so can be safely simplified.
 //
 // These heuristics might need more tuning
-static const UStruct* SkipEmptyBases(const UStruct* In)
+const UStruct* SkipEmptyBases(const UStruct* In)
 {
 	const UStruct* FirstOwner = In->PropertyLink ? In->PropertyLink->GetOwnerChecked<UStruct>() : In;
 	if (In != FirstOwner)
@@ -3404,8 +3380,6 @@ FStructDeclarationPtr FPropertyBinder::Declare() const
 
 static EPropertyFlags GDontBind;
 
-static bool ShouldBind(const UStruct* Struct);
-
 static bool ShouldBind(const FProperty* Property)
 {
 	if (Property->PropertyFlags & GDontBind)
@@ -3421,7 +3395,7 @@ static bool ShouldBind(const FProperty* Property)
 	return true;
 }
 
-static bool ShouldBind(const UStruct* Struct)
+bool ShouldBind(const UStruct* Struct)
 {
 	for (FProperty* Property = Struct->PropertyLink; Property; Property = Property->PropertyLinkNext)
 	{
@@ -3558,1663 +3532,49 @@ static void InitBatchedProperties()
 	GUE.Defaults.BindZeroes(GUE.Structs.LazyObjectPtr, sizeof(FLazyObjectPtr), alignof(FLazyObjectPtr));
 }
 
-struct FMemoryPropertyBatch
-{
-	TArray<FText> Texts; // Tricky to serialize intrusively
-
-	static FSoleMemberSpec Spec(FName*, FDeclId Id)
-	{
-		return {Id, GUE.Members.Id, Specify<uint64>()};
-	}
-
-	static void Save(FMemberBuilder& Out, FName In, const FSaveContext&)
-	{
-		Out.Add(GUE.Members.Id, FSensitiveName(In).ToUnstableInt());
-	}
-
-	static void Load(FName& Out, FStructLoadView In)
-	{
-		Out = FSensitiveName::FromUnstableInt(LoadSole<FSensitiveName::IntType>(In)).ToName();
-	}
-
-	static FSoleMemberSpec Spec(FText*, FDeclId Id)
-	{
-		return {Id, GUE.Members.Id, Specify<int32>()};
-	}
-
-	void Save(FMemberBuilder& Out, const FText& In, const FSaveContext&)
-	{
-		Out.Add(GUE.Members.Id, Texts.Num());
-		Texts.Add(In);
-	}
-
-	void Load(FText& Out, FStructLoadView In) const
-	{
-		Out = Texts[LoadSole<int32>(In)];
-	}
-
-	static FSoleMemberSpec Spec(FObjectHandle*, FDeclId Id)
-	{
-		return {Id, GUE.Members.Id, Specify<uint64>()};
-	}
-
-	static void Save(FMemberBuilder& Out, FObjectHandle In, const FSaveContext&)
-	{
-		static_assert(sizeof(In) == sizeof(uint64));
-		Out.Add(GUE.Members.Id, reinterpret_cast<const uint64&>(In));
-	}
-
-	static void Load(FObjectHandle& Out, FStructLoadView In)
-	{
-		LoadSole<uint64>(&Out, In);
-	}
-
-	static FSoleMemberSpec Spec(FWeakObjectPtr*, FDeclId Id)
-	{
-		return {Id, GUE.Members.Id, Specify<uint64>()};
-	}
-
-	static void Save(FMemberBuilder& Out, const FWeakObjectPtr& In, const FSaveContext&)
-	{
-		// Save ObjectSerialNumber + ObjectIndex a single uint64
-		static_assert(sizeof(FWeakObjectPtr) == sizeof(uint64));
-		Out.Add(GUE.Members.Id, reinterpret_cast<const uint64&>(In));
-	}
-
-	static void Load(FWeakObjectPtr& Out, FStructLoadView In)
-	{
-		LoadSole<uint64>(&Out, In);
-	}
-
-	static FSoleMemberSpec Spec(FSoftObjectPtr*, FDeclId Id)
-	{
-		return {Id, GUE.Members.Id, FMemberSpec(GUE.Structs.SoftObjectPath)};
-	}
-
-	static void Save(FMemberBuilder& Out, const FSoftObjectPtr& In, const FSaveContext& Ctx)
-	{
-		FBuiltStruct* SoftPath = SaveStruct(&In.GetUniqueID(), GUE.Structs.SoftObjectPath, Ctx);
-		Out.AddStruct(GUE.Members.Id, GUE.Structs.SoftObjectPath, SoftPath);
-	}
-
-	static void Load(FSoftObjectPtr& Out, FStructLoadView In)
-	{
-		Out.ResetWeakPtr();
-		LoadSoleStruct(&Out.GetUniqueID(), In);
-	}
-	
-	static FSoleMemberSpec Spec(FLazyObjectPtr*, FDeclId Id)
-	{
-		return {Id, GUE.Members.Id, FMemberSpec(GUE.Structs.Guid)};
-	}
-	
-	static void Save(FMemberBuilder& Out, const FLazyObjectPtr& In, const FSaveContext& Ctx)
-	{
-		FBuiltStruct* Guid = SaveStruct(&In.GetUniqueID(), GUE.Structs.Guid, Ctx);
-		Out.AddStruct(GUE.Members.Id, GUE.Structs.Guid, Guid);
-	}
-
-	static void Load(FLazyObjectPtr& Out, FStructLoadView In)
-	{
-		Out.ResetWeakPtr();
-		LoadSoleStruct(&Out.GetUniqueID(), In);
-	}
-};
-
-inline bool DiffProperty(FName A, FName B) { return !A.IsEqual(B, ENameCase::CaseSensitive); }
-inline bool DiffProperty(const FText& A, const FText& B) { return !FTextProperty::Identical_Implementation(A, B, 0); }
-inline bool DiffProperty(FObjectHandle A, FObjectHandle B) { return A != B; }
-inline bool DiffProperty(const FWeakObjectPtr& A, const FWeakObjectPtr& B) { return A != B; }
-inline bool DiffProperty(const FSoftObjectPtr& A, const FSoftObjectPtr& B) { return A != B; }
-inline bool DiffProperty(const FLazyObjectPtr& A, const FLazyObjectPtr& B) { return A != B; }
-
-template<class Type, class BatchType>
-struct TCustomPropertyBinding final : ICustomBinding
-{
-	TCustomPropertyBinding(BatchType& InBatch) : Batch(InBatch) {}
-
-	BatchType& Batch;
-
-	virtual void SaveCustom(FMemberBuilder& Dst, const void* Src, const void* Default, const FSaveContext& Ctx) override
-	{
-		if (!Default || DiffCustom(Src, Default, Ctx))
-		{
-			Batch.Save(Dst, *static_cast<const Type*>(Src), Ctx);
-		}
-	}
-
-	virtual void LoadCustom(void* Dst, FStructLoadView Src, ECustomLoadMethod) const override
-	{
-		Batch.Load(*static_cast<Type*>(Dst), Src);
-	}
-
-	virtual bool DiffCustom(const void* A, const void* B, const FBindContext&) const override
-	{
-		return DiffProperty(*static_cast<const Type*>(A), *static_cast<const Type*>(B));
-	}
-};
-
-template<class BatchType>
-struct TCustomPropertyBindings
-{
-	TCustomPropertyBindings(BatchType& Batch, const FCustomBindings& Underlay)
-	: Overlay(Underlay)
-	, Name(Batch)
-	, Text(Batch)
-	, ObjectPtr(Batch)
-	, SoftObjectPtr(Batch)
-	, WeakObjectPtr(Batch)
-	, LazyObjectPtr(Batch)
-	{
-		Bind(GUE.Structs.Name, Name);
-		Bind(GUE.Structs.Text, Text);
-		Bind(GUE.Structs.ClassPtr, ObjectPtr); // TSubclassOf<> is essentially a TObjectPtr
-		Bind(GUE.Structs.ObjectPtr, ObjectPtr);
-		Bind(GUE.Structs.SoftObjectPtr, SoftObjectPtr);
-		Bind(GUE.Structs.WeakObjectPtr, WeakObjectPtr);
-		Bind(GUE.Structs.LazyObjectPtr, LazyObjectPtr);
-	}
-
-	template<class Type>
-	void Bind(FDualStructId Id, TCustomPropertyBinding<Type, BatchType>& Binding)
-	{
-		Overlay.BindStruct(Id, Binding, Declare(BatchType::Spec((Type*)nullptr, Id)), {});
-	}
-
-	FCustomBindingsOverlay								Overlay;
-	TCustomPropertyBinding<FName, BatchType>			Name;
-	TCustomPropertyBinding<FText, BatchType>			Text;
-	TCustomPropertyBinding<FObjectHandle, BatchType>	ObjectPtr;
-	TCustomPropertyBinding<FSoftObjectPtr, BatchType>	SoftObjectPtr;
-	TCustomPropertyBinding<FWeakObjectPtr, BatchType>	WeakObjectPtr;
-	TCustomPropertyBinding<FLazyObjectPtr, BatchType>	LazyObjectPtr;
-};
-
-////////////////////////////////////////////////////////////////////////////////////////////////
-
-struct FMemoryBatch
-{
-	TArray64<uint8>				Data;
-	TArray<FStructId>			RuntimeIds; // To avoid reindexing schema FType
-	FMemoryPropertyBatch		Properties; // Contains FTexts referenced from Data
-};
-
-static constexpr uint32 Magics[] = { 0xFEEDF00D, 0xABCD1234, 0xDADADAAA, 0x99887766, 0xF0F1F2F3 };
-volatile static const UObject* GDebugNoteObject;
-
-class FBatchSaver
-{
-public:
-	FBatchSaver(FCustomBindings& Customs, int32 NumReserve)
-	: FlatCtx{{GUE.Schemas, Customs}, Scratch, nullptr}
-	, DeltaCtx{{GUE.Schemas, Customs}, Scratch, &GUE.Defaults}
-	{
-		SavedObjects.Reserve(NumReserve);
-		ObjectMap.Reserve(NumReserve);
-	}
-
-	// If Arch, delta serialize the matching part of the object inheritance chain
-	FBuiltStruct* SaveImpl(FBindId Id, const UStruct* Class, const UObject* Object, const UStruct* ArchClass, const UObject* Arch)
-	{
-		if (ArchClass == nullptr)
-		{
-			return SaveStruct(Object, Id, FlatCtx);
-		}
-		if (Class == ArchClass)
-		{
-			return SaveStructDelta(Object, Arch, Id, DeltaCtx);
-		}
-
-		Class = SkipEmptyBases(Class->GetInheritanceSuper());
-		FBindId SuperId = GUE.Names.IndexBindId(IndexType(Class));
-		FBuiltStruct* Super = SaveImpl(SuperId, Class, Object, ArchClass, Arch);
-		return SaveStructWithSuper(Object, Super, Id, FlatCtx);
-	}
-
-	void Save(FBindId Id, const UObject* Object, const UObject* Arch)
-	{
-		FBuiltStruct* Built;
-		if (Arch)
-		{
-			const UStruct* Class = SkipEmptyBases(Object->GetClass());
-			const UStruct* ArchClass = SkipEmptyBases(Arch->GetClass());
-			Built = SaveImpl(Id, Class, Object, ArchClass, Arch);
-		}
-		else
-		{
-			Built = SaveImpl(Id, Object->GetClass(), Object, nullptr, nullptr);
-		}
-		ObjectMap.Add(SavedObjects.Num());
-		SavedObjects.Emplace(Id, Built, Object);
-	}
-	
-	void Skip(FBindId Id, const UObject* Object)
-	{
-		ObjectMap.Add(~0u);
-	}
-
-	// Use existing SaveNameBatch() API for now, can be optimized later
-	void WriteStableNames(TArray64<uint8>& Out, TConstArrayView<FNameId> UsedNames) const
-	{
-		TArray<FDisplayNameEntryId> Entries;
-		TArray<int32> Numbers;
-		Entries.Reserve(UsedNames.Num());
-		Numbers.Reserve(UsedNames.Num());
-		for (FNameId NameId : UsedNames)
-		{
-			FSensitiveName SensitiveName = GUE.Names.ResolveName(NameId);
-			FName Name = SensitiveName.ToName();
-			Entries.Add(FDisplayNameEntryId(Name));
-			Numbers.Add(Name.GetNumber());
-		}
-
-		FBufferWriter NameMapArchive(nullptr, 0, EBufferWriterFlags::AllowResize | EBufferWriterFlags::TakeOwnership);
-		SaveNameBatch(Entries, NameMapArchive);
-
-		WriteNumAndArray(Out, TArrayView<uint8>(static_cast<uint8*>(NameMapArchive.GetWriterData()), static_cast<int32>(NameMapArchive.Tell())));
-		WriteNumAndArray(Out, Numbers);
-	}
-
-	TArray64<uint8>	Write(TArray<FStructId>* OutRuntimeIds = nullptr) const
-	{
-		ESchemaFormat Format = OutRuntimeIds ? ESchemaFormat::InMemoryNames : ESchemaFormat::StableNames;
-
-		// Build partial schemas
-		const FBindDeclarations Types(GUE.Enums, FlatCtx.Customs, GUE.Schemas);
-		FSchemasBuilder SchemaBuilders(GUE.Names, Types, Scratch, Format);
-		for (const FSavedObject& Object : SavedObjects)
-		{
-			GDebugNoteObject = Object.Input;
-			SchemaBuilders.NoteStructAndMembers(Object.Id, *Object.Built);
-		}
-		GDebugNoteObject = nullptr;
-		FBuiltSchemas Schemas = SchemaBuilders.Build();
-
-		// Save schema ids on the side when using InMemoryNames
-		if (OutRuntimeIds)
-		{
-			*OutRuntimeIds = ExtractRuntimeIds(Schemas);
-		}
-
-		FWriter Writer(GUE.Names, Types, Schemas, Format);
-		TArray64<uint8> Out;
-
-		WriteInt(Out, Magics[0]);
-		// Write out FNames when using StableNames
-		if (!OutRuntimeIds)
-		{
-			WriteStableNames(Out, Writer.GetUsedNames());
-		}
-
-		// Write schemas
-		WriteInt(Out, Magics[1]);
-		WriteAlignmentPadding<uint32>(Out);
-		TArray64<uint8> Tmp;
-		Writer.WriteSchemas(/* Out */ Tmp);
-		WriteNumAndArray(Out, TArrayView<const uint8, int64>(Tmp));
-		Tmp.Reset();
-
-		// Write objects
-		WriteInt(Out, Magics[2]);
-		for (const FSavedObject& Object : SavedObjects)
-		{
-			WriteInt(/* out */ Tmp, Magics[3]);
-			WriteInt(/* out */ Tmp, Writer.GetWriteId(Object.Id).Get().Idx);
-			Writer.WriteMembers(/* out */ Tmp, Object.Id, *Object.Built);
-			WriteSkippableSlice(Out, Tmp);
-			Tmp.Reset();
-		}
-
-		// Write object terminator
-		WriteSkippableSlice(Out, TConstArrayView64<uint8>());
-		WriteInt(Out, Magics[4]);
-
-		// Write object map that allows for empty objects/(indices) without built data
-		WriteNumAndArray(Out, ObjectMap);
-		
-		return Out;
-	}
-
-private:
-	struct FSavedObject
-	{
-		FBindId				Id;
-		FBuiltStruct*		Built;
-		const UObject*		Input; // For debug
-	};
-
-	TArray<FSavedObject>		SavedObjects;
-	TArray<uint32>				ObjectMap;
-	mutable FScratchAllocator	Scratch;
-	FSaveContext				FlatCtx;
-	FSaveContext				DeltaCtx;
-
-	template<typename ArrayType>
-	static void WriteNumAndArray(TArray64<uint8>& Out, const ArrayType& Items)
-	{
-		WriteInt(Out, IntCastChecked<uint32>(Items.Num()));
-		WriteArray(Out, Items);
-	}
-};
-
-class FMemoryBatchLoader
+struct FPlainPropsBindingScope
 {
-public:
-	FMemoryBatchLoader(const FCustomBindings& Customs, FMemoryView Data, TConstArrayView<FStructId> RuntimeIds)
-	{
-		// Read and mount schemas
-		FByteReader It(Data);
-		verify(It.Grab<uint32>() == Magics[0]);
-		verify(It.Grab<uint32>() == Magics[1]);
-		It.SkipAlignmentPadding<uint32>();
-		uint32 SchemasSize = It.Grab<uint32>();
-		const FSchemaBatch* SavedSchemas = ValidateSchemas(It.GrabSlice(SchemasSize));
-		verify(It.Grab<uint32>() == Magics[2]);
-		
-		FSchemaBatchId Batch = MountReadSchemas(SavedSchemas);
-
-		// Read objects
- 		while (uint64 NumBytes = It.GrabVarIntU())
-		{	
-			FByteReader ObjIt(It.GrabSlice(NumBytes));
-			verify(ObjIt.Grab<uint32>() == Magics[3]);
-			FStructSchemaId Id = { ObjIt.Grab<uint32>() };
-			check(Id != FStructSchemaId{});
-			Objects.Add({ { Id, Batch }, ObjIt });
-		}
-		
-		verify(It.Grab<uint32>() == Magics[4]);
-		verify(GrabNumAndArray<uint32>(It).Num() == Objects.Num());
-		It.CheckEmpty();
-		check(!Objects.IsEmpty());
-
-		// Finally create load plans
-		Plans = CreateLoadPlans(Batch, Customs, GUE.Schemas, RuntimeIds, ESchemaFormat::InMemoryNames);
-	}
-
-	~FMemoryBatchLoader()
-	{
-		check(LoadIdx == Objects.Num()); // Test should load all saved objects
-		Plans.Reset();
-		UnmountReadSchemas(Objects[0].Schema.Batch);
-	}
-
-	FORCENOINLINE void Load(UObject* Dst)
-	{
-		FStructView In = Objects[LoadIdx];
-		LoadStruct(Dst, In.Values, In.Schema.Id, *Plans);
-		++LoadIdx;
-	}
-
-	FORCENOINLINE void Reload(UObject* Dst, int32 ReloadIdx)
+	FPlainPropsBindingScope()
+		: Transform()
+		, Guid()
+		, Color()
+		, LinearColor()
+		, FieldPath(GUE.Structs.FieldPath)
+		, Delegate(GUE.Structs.Delegate)
+		, InlineMulticast({ GUE.Structs.MulticastInlineDelegate, GUE.Structs.MulticastDelegate })
+		, VerseFunction(GUE.Structs.VerseFunction)
+		, DynamicallyTypedValue(GUE.Structs.DynamicallyTypedValue)
+		, ReferencePropertyValue(GUE.Structs.ReferencePropertyValue)
 	{
-		FStructView In = Objects[ReloadIdx];
-		LoadStruct(Dst, In.Values, In.Schema.Id, *Plans);
 	}
 
-private:
-	FLoadBatchPtr				Plans;
-	TArray<FStructView>			Objects;
-	int32						LoadIdx = 0;
-	
-	template<typename T>
-	static TConstArrayView<T> GrabNumAndArray(/* in-out */ FByteReader& It)
-	{
-		uint32 Num = It.Grab<uint32>();
-		return MakeArrayView(reinterpret_cast<const T*>(It.GrabBytes(Num * sizeof(T))), Num);
-	}
+	TScopedStructBinding<FTransform, FDefaultRuntime> Transform;
+	TScopedStructBinding<FGuid, FDefaultRuntime> Guid;
+	TScopedStructBinding<FColor, FDefaultRuntime> Color;
+	TScopedStructBinding<FLinearColor, FDefaultRuntime> LinearColor;
+	TScopedStructBinding<FFieldPath, FDefaultRuntime> FieldPath;
+	TScopedStructBinding<FScriptDelegate, FDefaultRuntime> Delegate;
+	// MulticastDelegate declaration is shared with MulticastSparseDelegate
+	TScopedStructBinding<FMulticastScriptDelegate, FDefaultRuntime> InlineMulticast;
+	// Verse
+	TScopedStructBinding<FVerseFunction, FDefaultRuntime> VerseFunction;
+	TScopedStructBinding<::UE::FDynamicallyTypedValue, FDefaultRuntime> DynamicallyTypedValue;
+	TScopedStructBinding<FReferencePropertyValue, FDefaultRuntime> ReferencePropertyValue;
 };
 
-class FPersistentBatchLoader
-{
-public:
-	FPersistentBatchLoader(const FCustomBindings& Customs, FMemoryView Data)
-	{
-		// Read ids
-		FByteReader It(Data);
-		verify(It.Grab<uint32>() == Magics[0]);
-		ReadStableNames(It);
-
-		// Read and mount schemas
-		verify(It.Grab<uint32>() == Magics[1]);
-		It.SkipAlignmentPadding<uint32>();
-		uint32 SchemasSize = It.Grab<uint32>();
-		const FSchemaBatch* SavedSchemas = ValidateSchemas(It.GrabSlice(SchemasSize));
-		verify(It.Grab<uint32>() == Magics[2]);
-		
-		// Bind saved ids to runtime ids, make new schemas with new ids and mount them
-		FIdTranslator RuntimeIds(GUE.Names, MakeConstArrayView(Names), *SavedSchemas);
-		TranslatedSchemas = CreateTranslatedSchemas(*SavedSchemas, RuntimeIds.Translation);
-		FSchemaBatchId Batch = MountReadSchemas(TranslatedSchemas);
-
-		// Read objects
- 		while (uint64 NumBytes = It.GrabVarIntU())
-		{	
-			FByteReader ObjIt(It.GrabSlice(NumBytes));
-			verify(ObjIt.Grab<uint32>() == Magics[3]);
-			FStructSchemaId Id = { ObjIt.Grab<uint32>() };
-			Objects.Add({ { Id, Batch }, ObjIt });
-		}
-		
-		verify(It.Grab<uint32>() == Magics[4]);
-		ObjectMap = GrabNumAndArray<uint32>(It);
-		check(ObjectMap.Num() >= Objects.Num());
-		check(!Objects.IsEmpty());
-
-		// Finally create load plans
-		TConstArrayView<FStructId> LoadStructIds = RuntimeIds.Translation.GetStructIds(NumStructSchemas(Batch));
-		Plans = CreateLoadPlans(Batch, Customs, GUE.Schemas, LoadStructIds, ESchemaFormat::StableNames);
-	}
-
-	void ReadStableNames(FByteReader& It)
-	{
-		TConstArrayView64<uint8> NameBatch = GrabNumAndArray<uint8>(It);
-		TConstArrayView<int32> NameNumbers = GrabNumAndArray<int32>(It);
-		FMemoryReaderView Ar(MakeArrayView(NameBatch.GetData(), NameBatch.Num()));
-		TArray<FDisplayNameEntryId> NameEntries = LoadNameBatch(Ar);
-		Names.Reserve(NameEntries.Num());
-		check(NameEntries.Num() == NameNumbers.Num());
-		for (int64 Idx = 0, Num = NameEntries.Num(); Idx < Num; ++Idx)
-		{
-			Names.Emplace(NameEntries[Idx].ToName(NameNumbers[Idx]));
-		}
-	}
-
-	~FPersistentBatchLoader()
-	{
-		Plans.Reset();
-		UnmountReadSchemas(Objects[0].Schema.Batch);
-	}
-
-	FORCENOINLINE void Load(UObject* Dst)
-	{
-		FStructView In = Objects[LoadIdx];
-		LoadStruct(Dst, In.Values, In.Schema.Id, *Plans);
-		++LoadIdx;
-	}
-
-	bool HasObject(int32 ObjectIndex)
-	{
-		return ObjectMap[ObjectIndex] != ~0u;
-	}
-
-	FORCENOINLINE void LoadObject(UObject* Dst, int32 ObjectIndex)
-	{
-		FStructView In = Objects[ObjectMap[ObjectIndex]];
-		LoadStruct(Dst, In.Values, In.Schema.Id, *Plans);
-	}
-
-private:
-	const FSchemaBatch*			TranslatedSchemas = nullptr;
-	TArray<FSensitiveName>		Names;
-	FLoadBatchPtr				Plans;
-	TArray<FStructView>			Objects;
-	TConstArrayView<uint32>		ObjectMap;
-	int32						LoadIdx = 0;
-
-	template<typename T>
-	static TConstArrayView<T> GrabNumAndArray(/* in-out */ FByteReader& It)
-	{
-		uint32 Num = It.Grab<uint32>();
-		return MakeArrayView(reinterpret_cast<const T*>(It.GrabBytes(Num * sizeof(T))), Num);
-	}
-};
+} // namespace PlainProps::UE
 
 ////////////////////////////////////////////////////////////////////////////////////////////////
 
-// Similar to PlainProps FMemoryBatch but for storing FArchive property serialization results
-struct FArchivedProperties
-{
-	TArray<uint8>			Data;
-	TArray<FText>			Texts;
-};
-
-static constexpr uint32 RoundtripPortFlags = PPF_UseDeprecatedProperties | PPF_ForceTaggedSerialization;
-
-// Match FMemoryPropertyBatch somewhat for a fair comparison, e.g. save FText on side and FName as integer
-class FPropertyWriter final : public FMemoryWriter
-{
-	TArray<FText>& Texts;
-public:
-	FPropertyWriter(FArchivedProperties& Out) 
-	: FMemoryWriter(Out.Data)
-	, Texts(Out.Texts)
-	{
-		SetPortFlags(RoundtripPortFlags);
-	}
-
-	FORCENOINLINE void WriteProperties(UObject* Object, UObject* Defaults)
-	{
-		UClass* Class = Object->GetClass();
-		Class->SerializeTaggedProperties(*this, reinterpret_cast<uint8*>(Object), Class, reinterpret_cast<uint8*>(Defaults));
-	}
-
-	virtual FArchive& operator<<(FText& Value) override
-	{
-		int32 Idx = INDEX_NONE;
-		if (!Value.IsEmpty())
-		{
-			Idx = Texts.Num();
-			Texts.Add(Value);
-		}
-		return WriteValue(Idx);
-	}
-	virtual FArchive& operator<<(FName& Value) override				{ return WriteValue(FSensitiveName(Value).ToUnstableInt());}
-	virtual FArchive& operator<<(UObject*& Value) override			{ return WriteValue(reinterpret_cast<uint64&>(Value)); }
-	virtual FArchive& operator<<(FObjectPtr& Value) override		{ return WriteValue(reinterpret_cast<uint64&>(Value)); }
-	virtual FArchive& operator<<(FWeakObjectPtr& Value) override	{ return WriteValue(reinterpret_cast<uint64&>(Value)); }
-	virtual FArchive& operator<<(FLazyObjectPtr& Value) override	{ return WriteValue(Value.GetUniqueID()); }
-	virtual FArchive& operator<<(FSoftObjectPtr& Value) override	{ Value.GetUniqueID().SerializePath(*this); return *this; }
-	virtual FArchive& operator<<(FSoftObjectPath& Value) override	{ Value.SerializePath(*this); return *this; }
-	virtual FString GetArchiveName() const override					{ return "FPropertyWriter"; }
-
-	template<typename T>
-	FArchive& WriteValue(T Value)
-	{
-		return *this << Value;
-	}
-};
+static TUniquePtr<PlainProps::UE::FPlainPropsBindingScope> BindingScopePtr = nullptr;
+static TUniquePtr<PlainProps::DbgVis::FIdScope> DebugIdScopePtr = nullptr;
 
-class FPropertyReader final : public FMemoryReader
-{
-	const TArray<FText>& Texts;
-public:
-	FPropertyReader(const FArchivedProperties& Out)
-	: FMemoryReader(Out.Data)
-	, Texts(Out.Texts)
-	{
-		SetPortFlags(RoundtripPortFlags);
-	}
-
-	FORCENOINLINE void ReadProperties(UObject* Object, UObject* Defaults)
-	{
-		UClass* Class = Object->GetClass();
-		Class->SerializeTaggedProperties(*this, reinterpret_cast<uint8*>(Object), Class, reinterpret_cast<uint8*>(Defaults));
-	}
-
-	virtual FArchive& operator<<(FText& Value) override
-	{ 
-		int32 Idx = ReadValue<int32>();
-		Value = Idx == INDEX_NONE ? FText::GetEmpty() : Texts[Idx];
-		return *this;
-	}
-
-	virtual FArchive& operator<<(FName& Value) override
-	{
-		Value = FSensitiveName::FromUnstableInt(ReadValue<FSensitiveName::IntType>()).ToName();
-		return *this;
-	}
-
-	virtual FArchive& operator<<(UObject*& Value) override			{ reinterpret_cast<uint64&>(Value) = ReadValue<uint64>(); return *this; }
-	virtual FArchive& operator<<(FObjectPtr& Value) override		{ reinterpret_cast<uint64&>(Value) = ReadValue<uint64>(); return *this; }
-	virtual FArchive& operator<<(FWeakObjectPtr& Value) override	{ reinterpret_cast<uint64&>(Value) = ReadValue<uint64>(); return *this; }
-	virtual FArchive& operator<<(FLazyObjectPtr& Value) override	{ Value.ResetWeakPtr(); Value.GetUniqueID() = ReadValue<FUniqueObjectGuid>(); return *this; }
-	virtual FArchive& operator<<(FSoftObjectPtr& Value) override	{ Value.ResetWeakPtr(); Value.GetUniqueID().SerializePath(*this); return *this; }
-	virtual FArchive& operator<<(FSoftObjectPath& Value) override	{ Value.SerializePath(*this); return *this; }
-	virtual FString GetArchiveName() const override					{ return "FPropertyReader"; }
-
-	template<typename T>
-	T ReadValue()
-	{
-		T Out;
-		*this << Out;
-		return Out;
-	}
-};
-
-////////////////////////////////////////////////////////////////////////////////////////////////
-
-struct FInstance
-{
-	FString PathName;
-	UObject* Orig;
-	UObject* Arch;
-	UObject* Base;
-	UObject* PP;
-	UObject* TPS;
-	UObject* UPS;
-	FBindId Id;
-
-	void Init()
-	{
-		UClass* Class = Orig->GetClass();
-		Class->GetDefaultObject(/* create lazily */ true);
-		Arch = Orig->GetArchetype();
-		check(Arch != Orig);
-		Id = GUE.Names.IndexBindId(IndexType(Class));
-	}
-};
-
-static UObject* MakeEmptyInstance(UObject* Obj, FName Name)
-{
-	FStaticConstructObjectParameters Params(Obj->GetClass());
-	Params.Outer = Obj->GetOuter();//GetTransientPackage();
-	Params.Name = Name;
-	Params.SetFlags = Obj->GetFlags();
-	Params.Template = Obj->GetArchetype();
-	Params.bAssumeTemplateIsArchetype = true;
-	Params.bCopyTransientsFromClassDefaults = true;
-	return StaticConstructObject_Internal(Params);
-}
-
-static bool IncludeClass(UClass* Class)
-{
-	static const FName Exclusions[] = {
-		"CitySampleUnrealEdEngine", // Cloning MTAccessDetector crash
-		"GameFeaturePluginStateMachine", // Cloning ensure
-		"WorldSettings", // QAGame
-		// CitySample - Enum value 8 is undeclared in /Script/Engine.ERichCurveTangentMode, illegal value detected in /Script/Engine.RichCurveKey::TangentMode
-		"AnimSequence", "AnimationSequencerDataModel", "MovieSceneControlRigParameterSection",
-		// Enum value 7 is undeclared in /Script/Text3D.EText3DFontStyleFlags, illegal value detected in /Script/Text3DEditor.Text3DEditorFont::FontStyleFlags
-		// FontStyleFlags	Monospace | Bold | Italic (7 '\a')	UnrealEditor-Text3DEditor.dll!EText3DFontStyleFlags
-		"Text3DEditorFontSubsystem",
-	};
-	static const FName SuperExclusions[] = { "LevelScriptActor" };
-
-	if (!!Algo::Find(Exclusions, Class->GetFName()))
-	{
-		return false;
-	}
-
-	// Exclude IDOs
-	static constexpr EClassFlags IdoFlags = CLASS_NotPlaceable | CLASS_Hidden | CLASS_HideDropDown;
-	if (Class->HasAllClassFlags(IdoFlags))
-	{
-		return false;
-	}
-
-	for (UStruct* Super = Class->GetInheritanceSuper(); Super; Super = Super->GetInheritanceSuper())
-	{
-		if (!!Algo::Find(SuperExclusions, Super->GetFName()))
-		{
-			return false;
-		}
-	}
-	return true;
-}
-
-FORCENOINLINE static void SavePlainProps(FBatchSaver& Batch, TConstArrayView<FInstance> Instances)
-{
-	for (const FInstance& Instance : Instances)
-	{	
-		Batch.Save(Instance.Id, Instance.Orig, Instance.Base);
-	}
-}
-
-FORCENOINLINE static void LoadPlainProps(FMemoryBatchLoader& Batch, TConstArrayView<FInstance> Instances)
-{
-	for (const FInstance& Instance : Instances)
-	{	
-		Batch.Load(Instance.PP);
-	}
-}
-
-template<bool bUps>
-FORCENOINLINE void SaveArchive(FPropertyWriter& Archive, TConstArrayView<FInstance> Instances)
-{
-	Archive.SetUseUnversionedPropertySerialization(bUps);
-	for (const FInstance& Instance : Instances)
-	{
-		Archive.WriteProperties(Instance.Orig, Instance.Base);
-	}
-}
-
-template<bool bUps>
-FORCENOINLINE void LoadArchive(FPropertyReader& Archive, TConstArrayView<FInstance> Instances)
-{
-	Archive.SetUseUnversionedPropertySerialization(bUps);
-	for (const FInstance& Instance : Instances)
-	{
-		Archive.ReadProperties(bUps ? Instance.UPS : Instance.TPS, Instance.Base);
-	}
-}
-
-////////////////////////////////////////////////////////////////////////////////////////////////
-
-class FStableNameBatchIds final : public FStableBatchIds
-{
-	TConstArrayView<FSensitiveName> Names;
-public:
-	FStableNameBatchIds(FSchemaBatchId Batch, TConstArrayView<FSensitiveName> InNames) : FStableBatchIds(Batch), Names(InNames) {}
-	using FStableBatchIds::AppendString;
-
-	virtual uint32				NumNames() const override							{ return static_cast<uint32>(Names.Num()); }
-	virtual void				AppendString(FUtf8Builder& Out, FNameId Name) const override { Names[Name.Idx].AppendString(Out); }
-};
-
-[[nodiscard]] static FSchemaBatchId ParseBatch(
-	TArray64<uint8>& OutData,
-	TArray<FStructView>& OutObjects,
-	FUtf8StringView YamlView)
-{
-	// Parse yaml
-	ParseYamlBatch(OutData, YamlView);
-
-	// Grab and mount parsed schemas
-	FByteReader It(MakeMemoryView(OutData));
-	const uint32 SchemasSize = It.Grab<uint32>();
-	FMemoryView SchemasView = It.GrabSlice(SchemasSize);
-	const FSchemaBatch* Schemas = ValidateSchemas(SchemasView);
-	FSchemaBatchId Batch = MountReadSchemas(Schemas);
-
-	// Grab parsed objects
-	while (uint64 NumBytes = It.GrabVarIntU())
-	{	
-		FByteReader ObjIt(It.GrabSlice(NumBytes));
-		FStructSchemaId Schema = { ObjIt.Grab<uint32>() };
-		OutObjects.Add({ { Schema, Batch }, ObjIt });
-	}
-	
-	return Batch;
-}
-
-static void RoundtripText(
-	const FBatchIds& BatchIds,
-	TConstArrayView<FStructView> Objects,
-	TConstArrayView<FInstance> Instances,
-	ESchemaFormat Format,
-	FStringView AssetName = {})
-{
-	// Print yaml
-	UE_LOGFMT(LogPlainPropsUObject, Display, "Printing to PlainProps text using {Format}...", ToString(Format));
-	TUtf8StringBuilder<256> Yaml;
-	Yaml.Reserve(INT_MAX);
-	PrintYamlBatch(Yaml, BatchIds, Objects);
-	FUtf8StringView YamlView = Yaml.ToView();
-
-	// Write to file
-	if (AssetName.Len() == 0)
-	{
-		AssetName = (Format == ESchemaFormat::InMemoryNames ? TEXTVIEW("DiffTest_InMemoryNames") : TEXTVIEW("DiffTest_StableNames"));
-	}
-	const FString YamlPath = FPaths::Combine(FPaths::ProjectSavedDir(), TEXT("PlainProps"), AssetName) + TEXT(".yaml");
-	if (TUniquePtr<FArchive> FileWriter(IFileManager::Get().CreateFileWriter(*YamlPath)); FileWriter)
-	{
-		UE_LOGFMT(LogPlainPropsUObject, Display, "Writing {KB}KB yaml as {YamlPath}...", Yaml.Len() >> 10, *YamlPath);
-		FileWriter->Serialize((void*)YamlView.GetData(), YamlView.Len());
-	}
-
-	// Parse yaml
-	UE_LOGFMT(LogPlainPropsUObject, Display, "Parsing PlainProps text using {Format}...", ToString(Format));
-	TArray64<uint8> Data;
-	TArray<FStructView> ParsedObjects;
-	FSchemaBatchId ParsedBatch = ParseBatch(Data, ParsedObjects, YamlView);
-
-	if (Format == ESchemaFormat::StableNames && Instances.Num())
-	{
-		UE_LOGFMT(LogPlainPropsUObject, Display, "Diffing PlainProps parsed objects using {Format}...", ToString(Format));
-
-		// Diff schemas
-		check(!DiffSchemas(BatchIds.GetBatchId(), ParsedBatch));
-
-		// Diff objects
-		check(Objects.Num() == ParsedObjects.Num());
-		check(Objects.Num() == Instances.Num());
-		const int32 NumObjects = FMath::Min(Objects.Num(), ParsedObjects.Num());
-		uint32 NumDiffs = 0;
-		TUtf8StringBuilder<256> Diffs;
-		for (int32 I = 0; I < NumObjects; ++I)
-		{
-			FStructView In = Objects[I];
-			FStructView Parsed = ParsedObjects[I];
-			FReadDiffPath DiffPath;
-			if (DiffStruct(In, Parsed, DiffPath))
-			{
-				PrintDiff(Diffs, BatchIds, DiffPath);
-				Diffs.Append(" in ");
-				Diffs.Append(Instances[I].PathName);
-				Diffs.Append("\n");
-				++NumDiffs;
-			}
-		}
-		UE_LOGFMT(LogPlainPropsUObject, Display,
-			"Detected {Diffs} diffs in {Objs} PlainProps parsed objects from {KB}KB yaml text using StableNames\n{Diffs}",
-			NumDiffs, NumObjects, Yaml.Len() >> 10, Diffs.ToString());
-	}
-
-	// Unmount parsed schemas
-	UnmountReadSchemas(ParsedBatch);
-}
-
-class FBatchTextRoundtripper
-{
-public:
-	FBatchTextRoundtripper(FMemoryView Data, ESchemaFormat InFormat) : Format(InFormat)
-	{
-		FByteReader It(Data);
-		verify(It.Grab<uint32>() == Magics[0]);
-
-		// Read FNames when using Stable Names
-		if (Format == ESchemaFormat::StableNames)
-		{
-			ReadStableNames(It);
-		}
-		
-		// Read and mount schemas
-		verify(It.Grab<uint32>() == Magics[1]);
-		It.SkipAlignmentPadding<uint32>();
-		const uint32 SchemasSize = It.Grab<uint32>();
-		FMemoryView SavedSchemasView = It.GrabSlice(SchemasSize);
-		const FSchemaBatch* SavedSchemas = ValidateSchemas(SavedSchemasView);
-		verify(It.Grab<uint32>() == Magics[2]);
-		FSchemaBatchId Batch = MountReadSchemas(SavedSchemas);
-		
-		// Read objects
-		while (uint64 NumBytes = It.GrabVarIntU())
-		{	
-			FByteReader ObjIt(It.GrabSlice(NumBytes));
-			verify(ObjIt.Grab<uint32>() == Magics[3]);
-			FStructSchemaId Id = { ObjIt.Grab<uint32>() };
-			Objects.Add({ { Id, Batch }, ObjIt });
-		}
-		
-		verify(It.Grab<uint32>() == Magics[4]);
-		check(!Objects.IsEmpty());
-
-		// Create BatchIds 
-		if (Format == ESchemaFormat::StableNames)
-		{
-			BatchIds = MakeUnique<FStableNameBatchIds>(Batch, Names);
-		}
-		else
-		{
-			BatchIds = MakeUnique<FMemoryBatchIds>(Batch, GUE.Names);
-		}
-	}
-
-	void ReadStableNames(FByteReader& It)
-	{
-		TConstArrayView<uint8> NameBatch = GrabNumAndArray<uint8>(It);
-		TConstArrayView<int32> Numbers = GrabNumAndArray<int32>(It);
-		FMemoryReaderView Ar(NameBatch);
-		TArray<FDisplayNameEntryId> Entries = LoadNameBatch(Ar);
-		Names.Reserve(Entries.Num());
-		check(Entries.Num() == Numbers.Num());
-		for (int64 Idx = 0, Num = Entries.Num(); Idx < Num; ++Idx)
-		{
-			Names.Emplace(Entries[Idx].ToName(Numbers[Idx]));
-		}
-	}
-
-	~FBatchTextRoundtripper()
-	{
-		UnmountReadSchemas(BatchIds->GetBatchId());
-	}
-
-	void RoundtripTextForDiffTest(TConstArrayView<FInstance> Instances) const
-	{
-		PlainProps::UE::RoundtripText(*BatchIds, Objects, Instances, Format, {});
-	}
-
-	void RoundtripTextForPackage(FStringView PackageName) const
-	{
-		PlainProps::UE::RoundtripText(*BatchIds, Objects, {}, Format, PackageName);
-	}
-
-private:
-	TArray<FStructView>			Objects;
-	TArray<FSensitiveName>		Names;
-	TUniquePtr<FBatchIds>		BatchIds;
-	ESchemaFormat				Format;
-	
-	template<typename T>
-	static TConstArrayView<T> GrabNumAndArray(/* in-out */ FByteReader& It)
-	{
-		uint32 Num = It.Grab<uint32>();
-		return MakeArrayView(reinterpret_cast<const T*>(It.GrabBytes(Num * sizeof(T))), Num);
-	}
-};
-
-////////////////////////////////////////////////////////////////////////////////
-
-struct FDiffDebug
-{
-	TArray<FInstance*> Instances;
-	volatile const UTF8CHAR* Str;
-};
-static FDiffDebug GPpDiff, GTpsDiff, GUpsDiff;
-
-static void DumpUObjectToJsonFile(const UObject* Object)
-{
-	FString JsonString;
-	if (FJsonObjectConverter::UStructToJsonObjectString(Object->GetClass(), Object, JsonString))
-	{
-		TStringBuilder<256> JsonPath;
-		JsonPath << FPaths::Combine(FPaths::ProjectSavedDir(), TEXT("PlainProps"), Object->GetPackage()->GetName());
-		JsonPath << "." << Object->GetName() + ".json";
-
-		UE_LOGFMT(LogPlainPropsUObject, Display, "Writing {JsonPath}...", JsonPath);
-		FFileHelper::SaveStringToFile(JsonString, JsonPath.ToString(), FFileHelper::EEncodingOptions::ForceUTF8WithoutBOM);
-	}
-}
-
-////////////////////////////////////////////////////////////////////////////////
-
-#if WITH_EDITOR
-
-struct FLinkerPropertyBatch
-{
-	explicit FLinkerPropertyBatch(const FLinkerSave& InLinkerSave)
-		: LinkerSave(&InLinkerSave)
-	{
-	}
-
-	explicit FLinkerPropertyBatch(FLinkerLoad& InLinkerLoad)
-		: LinkerLoad(&InLinkerLoad)
-	{
-	}
-
-	static FSoleMemberSpec Spec(FName*, FDeclId Id)
-	{
-		return {Id, GUE.Members.Id, Specify<uint64>()};
-	}
-
-	void Save(FMemberBuilder& Out, FName In, const FSaveContext&)
-	{
-		// see FLinkerSave::MapName
-		int32 Index = LinkerSave->NameIndices.FindRef(In.GetDisplayIndex(), INDEX_NONE);
-		int32 Number = In.GetNumber();
-		uint64 Id;
-		if (Index == INDEX_NONE)
-		{
-			UE_LOGFMT(LogPlainPropsUObject, Warning, "FName {Name} is not mapped.", In);
-			Id = 6666666;
-		}
-		else
-		{
-			Id = Index != INDEX_NONE ? ((uint64(Number) << 32) | Index) : 6666666;
-		}
-		Out.Add(GUE.Members.Id, Id);
-	}
-
-	void Load(FName& Out, FStructLoadView In)
-	{
-		uint64 Id;
-		LoadSole<uint64>(&Id, In);
-		if (Id != 6666666)
-		{
-			int32 Index = Id & 0xFFFFFFFF;
-			int32 Number = Id >> 32;
-			FNameEntryId Entry = LinkerLoad->NameMap[Index];
-			FName Name = FName::CreateFromDisplayId(Entry, Number);
-			Out = Name;
-		}
-		else
-		{
-			Out = FName();
-		}
-	}
-
-	static FSoleMemberSpec Spec(FText*, FDeclId Id)
-	{
-		return {Id, GUE.Members.Id, Specify<int32>()};
-	}
-
-	void Save(FMemberBuilder& Out, const FText& In, const FSaveContext&)
-	{
-		Out.Add(GUE.Members.Id, Texts.Num());
-		Texts.Add(In);
-	}
-
-	void Load(FText& Out, FStructLoadView In) const
-	{
-		int32 Dummy = LoadSole<int32>(In);
-		Out = {};//TODO: Texts[LoadSole<int32>(In)];
-	}
-
-	static FSoleMemberSpec Spec(FObjectHandle*, FDeclId Id)
-	{
-		return {Id, GUE.Members.Id, Specify<int32>()};
-	}
-
-	void Save(FMemberBuilder& Out, FObjectHandle In, const FSaveContext&)
-	{
-		// see FLinkerSave::MapObject
-		UObject* Object = FObjectPtr(In).Get();
-		FPackageIndex Index = LinkerSave->ObjectIndicesMap.FindRef(Object);
-		int32 Idx = *reinterpret_cast<int32*>(&Index);
-		if (Object && Index.IsNull())
-		{
-			UE_LOGFMT(LogPlainPropsUObject, Warning, "UObject {Object} is not mapped.", Object->GetFullName());
-			Idx = 9999999;
-		}
-		Out.Add(GUE.Members.Id, Idx);
-	}
-
-	void Load(FObjectHandle& Out, FStructLoadView In)
-	{
-		int32 Idx;
-		LoadSole<int32>(&Idx, In);
-		if (Idx != 9999999)
-		{
-			FPackageIndex PackageIdx = *reinterpret_cast<FPackageIndex*>(&Idx);
-			if (PackageIdx.IsImport())
-			{
-				UObject* Object = LinkerLoad->ImportMap[PackageIdx.ToImport()].XObject;
-				if (Object == nullptr)
-				{
-#if PP_WITH_ENGINE_HOOKS
-					// Object = LinkerLoad->ResolveResource(PackageIdx);
-#endif
-				}
-				Out = ::UE::CoreUObject::Private::MakeObjectHandle(Object);
-			}
-			else if (PackageIdx.IsExport())
-			{
-				UObject* Object = LinkerLoad->ExportMap[PackageIdx.ToExport()].Object;
-				if (Object == nullptr)
-				{
-#if PP_WITH_ENGINE_HOOKS
-					// Object = LinkerLoad->ResolveResource(PackageIdx);
-#endif
-				}
-				Out = ::UE::CoreUObject::Private::MakeObjectHandle(Object);
-			}
-			else
-			{
-				check(IsObjectHandleNull(Out));
-			}
-		}
-		else
-		{
-			Out = FObjectHandle();
-		}
-	}
-
-	static FSoleMemberSpec Spec(FWeakObjectPtr*, FDeclId Id)
-	{
-		return {Id, GUE.Members.Id, Specify<uint64>()};
-	}
-
-	void Save(FMemberBuilder& Out, const FWeakObjectPtr& In, const FSaveContext& Ctx)
-	{
-		// Save ObjectSerialNumber + ObjectIndex a single uint64
-		UObject* Object = In.Get(true);
-		Save(Out, ::UE::CoreUObject::Private::MakeObjectHandle(Object), Ctx);
-	}
-
-	void Load(FWeakObjectPtr& Out, FStructLoadView In)
-	{
-		FObjectHandle Handle;
-		Load(Handle, In);
-		Out = ::UE::CoreUObject::Private::ResolveObjectHandle(Handle);
-	}
-
-	static FSoleMemberSpec Spec(FSoftObjectPtr*, FDeclId Id)
-	{
-		return {Id, GUE.Members.Id, FMemberSpec(GUE.Structs.SoftObjectPath)};
-	}
-
-	static void Save(FMemberBuilder& Out, const FSoftObjectPtr& In, const FSaveContext& Ctx)
-	{
-		FBuiltStruct* SoftPath = SaveStruct(&In.GetUniqueID(), GUE.Structs.SoftObjectPath, Ctx);
-		Out.AddStruct(GUE.Members.Id, GUE.Structs.SoftObjectPath, SoftPath);
-	}
-
-	static void Load(FSoftObjectPtr& Out, FStructLoadView In)
-	{
-		Out.ResetWeakPtr();
-		LoadSoleStruct(&Out.GetUniqueID(), In);
-	}
-
-	static FSoleMemberSpec Spec(FLazyObjectPtr*, FDeclId Id)
-	{
-		return {Id, GUE.Members.Id, FMemberSpec(GUE.Structs.Guid)};
-	}
-	
-	static void Save(FMemberBuilder& Out, const FLazyObjectPtr& In, const FSaveContext& Ctx)
-	{
-		FBuiltStruct* Guid = SaveStruct(&In.GetUniqueID(), GUE.Structs.Guid, Ctx);
-		Out.AddStruct(GUE.Members.Id, GUE.Structs.Guid, Guid);
-	}
-
-	static void Load(FLazyObjectPtr& Out, FStructLoadView In)
-	{
-		Out.ResetWeakPtr();
-		LoadSoleStruct(&Out.GetUniqueID(), In);
-	}
-
-private:
-	union
-	{
-		const FLinkerSave* LinkerSave;
-		FLinkerLoad* LinkerLoad;
-		const FLinker* Linker;
-	};
-
-	TArray<FText> Texts; // Tricky to serialize non-intrusively
-};
-
-struct FPlainPropsLinkerSandbox
-{
-	static constexpr uint32 LoadFlags = (LOAD_Async | LOAD_NoVerify | LOAD_SkipLoadImportedPackages);
-	static constexpr uint32 CloneNumber = 8848;
-
-	struct FSavedPackage
-	{
-		FName PackageName;
-		FPackagePath PackagePath;
-		FSharedBuffer PackageData;
-		TArray64<uint8> PlainPropsData;
-	};
-
-	struct FLinkerBatch
-	{
-		FLinkerBatch(const FSavedPackage& Saved)
-			: PackagePath(Saved.PackagePath)
-			, PackageName(Saved.PackageName.GetNumber() ? FName(*Saved.PackageName.ToString(), CloneNumber) : FName(Saved.PackageName, CloneNumber))
-			, Package(NewObject<UPackage>(/*Outer*/nullptr, PackageName))
-			, LinkerReader(static_cast<const uint8*>(Saved.PackageData.GetData()), Saved.PackageData.GetSize(), ELargeMemoryReaderFlags::Persistent, PackageName)
-			, LinkerLoad(FLinkerLoad::CreateLinker(FUObjectThreadContext::Get().GetSerializeContext(), Package, PackagePath, LoadFlags, &LinkerReader))
-			, Properties(*LinkerLoad)
-			, Customs(Properties, GUE.Customs)
-			, BatchView(MakeMemoryView(Saved.PlainPropsData))
-			, BatchLoader(Customs.Overlay, BatchView)
-		{
-		}
-
-		FPackagePath PackagePath;
-		FName PackageName;
-		UPackage* Package;
-		FLargeMemoryReader LinkerReader;
-		FLinkerLoad* LinkerLoad;
-		FLinkerPropertyBatch Properties;
-		TCustomPropertyBindings<FLinkerPropertyBatch> Customs;
-		FMemoryView BatchView;
-		FPersistentBatchLoader BatchLoader;
-	};
-
-	static TMap<FName, FSavedPackage> SavedPackages;
-	static TMap<FName, TUniquePtr<FLinkerBatch>> LoadBatches;
-
-	static void SavePlainProps(const FLinkerSave& LinkerSave)
-	{
-		FName PackageName = LinkerSave.LinkerRoot->GetFName();
-
-		UE_LOGFMT(LogPlainPropsUObject, Log, " Saving PlainProps for {PackageName}...", PackageName);
-		FLinkerPropertyBatch Properties(LinkerSave);
-		TCustomPropertyBindings<FLinkerPropertyBatch> Customs(Properties, GUE.Customs);
-		FBatchSaver BatchSaver(Customs.Overlay, LinkerSave.ExportMap.Num());
-		for (const FObjectExport& Export : LinkerSave.ExportMap)
-		{
-			const UStruct* Struct = SkipEmptyBases(Export.Object->GetClass());
-			FBindId Id = GUE.Names.IndexBindId(PlainProps::UE::IndexType(Struct));
-			if (GUE.Schemas.FindStruct(Id))
-			{
-				BatchSaver.Save(Id, Export.Object, Export.Object->GetArchetype());
-			}
-			else
-			{
-				BatchSaver.Skip(Id, Export.Object);
-			}
-		}
-		TArray64<uint8> StableData = BatchSaver.Write();
-
-		FSavedPackage& SavedPackage = SavedPackages.FindOrAdd(PackageName);
-		SavedPackage.PackageName = PackageName;
-		SavedPackage.PackagePath = LinkerSave.LinkerRoot->GetLoadedPath();
-		SavedPackage.PlainPropsData = MoveTemp(StableData);
-	}
-
-	static void SavePackageToMemory(const FLinkerSave& LinkerSave, uint8* Data, int64 Size)
-	{
-		FName PackageName = LinkerSave.LinkerRoot->GetFName();
-		FSavedPackage& SavedPackage = SavedPackages.FindChecked(PackageName);
-		SavedPackage.PackageData = FSharedBuffer::TakeOwnership(Data, Size, FMemory::Free);
-	}
-
-	static bool OpenPackageReader(const FPackagePath& PackagePath, FArchive*& Ar)
-	{
-		FName PackageName = PackagePath.GetPackageFName();
-		if (FSavedPackage* SavedData = SavedPackages.Find(PackageName))
-		{
-			Ar = new FLargeMemoryReader(static_cast<const uint8*>(SavedData->PackageData.GetData()), SavedData->PackageData.GetSize(), ELargeMemoryReaderFlags::Persistent, PackageName);
-			return true;
-		}
-		return false;
-	}
-
-	static void MountPackageFromMemory(FName PackageName)
-	{
-		if (FSavedPackage* SavedData = SavedPackages.Find(PackageName))
-		{
-			FLinkerBatch* LinkerBatch = new FLinkerBatch(*SavedData);
-			LoadBatches.Emplace(PackageName, LinkerBatch);
-		}
-	}
-
-	static void LoadPackageFromMemory(FName PackageName)
-	{
-		if (TUniquePtr<FLinkerBatch>* LinkerBatch = LoadBatches.Find(PackageName))
-		{
-			UE_LOGFMT(LogPlainPropsUObject, Log, " Loading PlainProps for {Package}...", PackageName);
-			LoadPackage((*LinkerBatch)->Package, (*LinkerBatch)->PackagePath, LoadFlags);
-		}
-	}
-
-	static bool LoadObject(FStructuredArchive::FSlot& Slot, uint8* Data)
-	{
-		UObject* Object = (UObject*)Data;
-		FName PackageName = Object->GetPackage()->GetFName();
-		if (PackageName.GetNumber() == CloneNumber)
-		{
-			PackageName.SetNumber(0);
-		}
-
-		TUniquePtr<FLinkerBatch>* LinkerBatch = LoadBatches.Find(PackageName);
-		if (!LinkerBatch)
-		{
-			return false;
-		}
-
-		FPersistentBatchLoader& BatchLoader = (*LinkerBatch)->BatchLoader;
-		FLinkerLoad* LinkerLoad = (*LinkerBatch)->LinkerLoad;
-		check(LinkerLoad == Object->GetLinker());
-
-		FUObjectSerializeContext* SerializeContext = FUObjectThreadContext::Get().GetSerializeContext();
-		check(Object == SerializeContext->SerializedObject);
-
-		const int32 ExportIndex = Object->GetLinkerIndex();
-		const FObjectExport& Export = LinkerLoad->ExportMap[ExportIndex];
-		check(Export.Object == Object);
-
-		BatchLoader.LoadObject(Object, ExportIndex);
-
-		const int32 Pos = Slot.GetUnderlyingArchive().Tell();
-		check(Pos == SerializeContext->SerializedObjectScriptStartOffset);
-		const int32 ScriptSerializationSize = Export.ScriptSerializationEndOffset - Export.ScriptSerializationStartOffset;
-		check(ScriptSerializationSize >= 0);
-		Slot.GetUnderlyingArchive().Seek(Pos + ScriptSerializationSize);
-		return true;
-	}
-};
-
-TMap<FName, FPlainPropsLinkerSandbox::FSavedPackage> FPlainPropsLinkerSandbox::SavedPackages;
-TMap<FName, TUniquePtr<FPlainPropsLinkerSandbox::FLinkerBatch>> FPlainPropsLinkerSandbox::LoadBatches;
-
-#endif
-
-////////////////////////////////////////////////////////////////////////////////
-
-int32 RoundtripViaBatch(TConstArrayView<UObject*> Objects, ERoundtrip Options)
-{
-	TArray<FInstance> Instances;
-	if (Objects.Num())
-	{
-		UE_LOGFMT(LogPlainPropsUObject, Display, "Gathering UObjects in Assets...");
-		for (const UObject* Object : Objects)
-		{
-			ForEachObjectWithPackage(Object->GetPackage(), [&Instances](UObject* Obj)
-			{
-				UClass* Class = Obj->GetClass();
-				if (ShouldBind(Class) && IncludeClass(Class))
-				{
-					Instances.Emplace(Obj->GetPathName(), Obj);
-				}
-				return true;
-			});
-		}
-	}
-	else
-	{
-		UE_LOGFMT(LogPlainPropsUObject, Display, "Gathering all non-empty UObjects...");
-		static constexpr EObjectFlags SkipFlags = RF_ClassDefaultObject | RF_MirroredGarbage | RF_InheritableComponentTemplate;
-		for (TObjectIterator<UObject> It(SkipFlags); It; ++It)
-		{
-			UObject* Object = *It;
-			UClass* Class = Object->GetClass();
-			if (ShouldBind(Class) && IncludeClass(Class))
-			{
-				Instances.Emplace(Object->GetPathName(), Object);
-			}
-		}
-	}
-
-	UE_LOGFMT(LogPlainPropsUObject, Display, "Sorting {Num} UObjects ...", Instances.Num());
-	Algo::Sort(Instances, [](const FInstance& A, const FInstance& B)
-	{
-		return FPlatformString::Strcmp(*A.PathName, *B.PathName) < 0;
-	});
-
-	// Create CDOs if needed and then clones for PP and TPS tests
-	UE_LOGFMT(LogPlainPropsUObject, Display, "Cloning {Num} UObjects up to 4 times...", Instances.Num());
-	for (FInstance& Instance : Instances)
-	{
-		Instance.Init();
-	}
-	FlushAsyncLoading();
-
-	for (FInstance& Instance : Instances)
-	{
-		uint32 N = 1 + &Instance - Instances.GetData();
-		Instance.Base = Instance.Arch ? MakeEmptyInstance(Instance.Orig, FName("Base", N)) : nullptr;
-		if (EnumHasAnyFlags(Options, ERoundtrip::PP))
-		{
-			Instance.PP = MakeEmptyInstance(Instance.Orig, FName("PP", N));
-		}
-		if (EnumHasAnyFlags(Options, ERoundtrip::TPS))
-		{
-			Instance.TPS = MakeEmptyInstance(Instance.Orig, FName("TPS", N));
-		}
-		if (EnumHasAnyFlags(Options, ERoundtrip::UPS))
-		{
-			Instance.UPS = MakeEmptyInstance(Instance.Orig, FName("UPS", N));
-		}
-	}
-
-	// Save
-	UE_LOGFMT(LogPlainPropsUObject, Display, "Saving UObjects to PlainProps with InMemoryNames...");
-	FMemoryBatch Plain;
-	TCustomPropertyBindings<FMemoryPropertyBatch> Customs(Plain.Properties, GUE.Customs);
-	{
-		FBatchSaver Batch(Customs.Overlay, GUObjectArray.GetObjectArrayNum());
-		SavePlainProps(Batch, Instances);
-		Plain.Data = Batch.Write(/* out */ &Plain.RuntimeIds);
-
-		if (EnumHasAnyFlags(Options, ERoundtrip::TextMemory))
-		{
-			FBatchTextRoundtripper MemoryBatch(MakeMemoryView(Plain.Data), ESchemaFormat::InMemoryNames);
-			MemoryBatch.RoundtripTextForDiffTest(Instances);
-		}
-		if (EnumHasAnyFlags(Options, ERoundtrip::TextStable))
-		{
-			UE_LOGFMT(LogPlainPropsUObject, Display, "Saving UObjects to PlainProps with StableNames...");
-			TArray64<uint8> StableData = Batch.Write();
-
-			FBatchTextRoundtripper StableBatch(MakeMemoryView(StableData), ESchemaFormat::StableNames);
-			StableBatch.RoundtripTextForDiffTest(Instances);
-		}
-	}
-
-	// Load
-	uint32 NumPpDiffs = 0;
-	if (EnumHasAnyFlags(Options, ERoundtrip::PP))
-	{
-		UE_LOGFMT(LogPlainPropsUObject, Display, "Loading UObjects from PlainProps...");
-		FMemoryBatchLoader Batch(Customs.Overlay, MakeMemoryView(Plain.Data), Plain.RuntimeIds);
-		LoadPlainProps(Batch, Instances);
-
-		// Diff original vs PlainProps
-		UE_LOGFMT(LogPlainPropsUObject, Display, "Diffing UObjects roundtripped via PlainProps...");
-		FDiffContext DiffCtx = {{ GUE.Schemas, Customs.Overlay }};
-		TUtf8StringBuilder<256> PpDiffPaths;
-		for (FInstance& Instance : Instances)
-		{
-			if (DiffStructs(Instance.Orig, Instance.PP, Instance.Id, /* in-out */ DiffCtx))
-			{
-				PpDiffPaths.Append("\n");
-				GPpDiff.Instances.Add(&Instance);
-				PrintDiff(/* out */ PpDiffPaths, GUE.Names, DiffCtx.Out);
-
-				//Batch.Reload(Instance.UPS, &Instance - Instances.GetData());
-				DiffCtx.Out.Reset();
-				PpDiffPaths.Append(" in ");
-				PpDiffPaths.Append(Instance.Orig->GetFullName());
-				++NumPpDiffs;
-			}
-		}
-		GPpDiff.Str = PpDiffPaths.GetData();
-		UE_LOGFMT(LogPlainPropsUObject, Display, "Detected {Diffs} diffs in {Objs} UObjects saved in a {KB}KB value stream using PlainProps{DiffText}", NumPpDiffs, Instances.Num(), Plain.Data.NumBytes() / 1024, PpDiffPaths.ToString());
-		for (const FInstance* Instance : GPpDiff.Instances)
-		{
-			DumpUObjectToJsonFile(Instance->Orig);
-			DumpUObjectToJsonFile(Instance->PP);
-		}
-	}
-
-	FArchivedProperties Tps;
-	if (EnumHasAnyFlags(Options, ERoundtrip::TPS))
-	{
-		UE_LOGFMT(LogPlainPropsUObject, Display, "Saving UObjects to TPS archive...");
-		FPropertyWriter Archive(/* out */ Tps);
-		SaveArchive<false>(Archive, Instances);
-	}
-
-	if (EnumHasAnyFlags(Options, ERoundtrip::TPS))
-	{
-		UE_LOGFMT(LogPlainPropsUObject, Display, "Loading UObjects from TPS archive...");
-		FPropertyReader Archive(Tps);
-		LoadArchive<false>(Archive, Instances);
-	}
-
-	static const FName SkipClasses[] = {
-		"BodySetup", // Skips some structs due to native FCollisionResponse::operator==
-		"NiagaraScript", "NiagaraNodeFunctionCall", "NiagaraMeshRendererProperties", // FNiagaraTypeDefinition::Serialize resets ClassStructOrEnum
-	};
-
-	if (EnumHasAnyFlags(Options, ERoundtrip::TPS))
-	{
-		// Diff original vs TPS
-		UE_LOGFMT(LogPlainPropsUObject, Display, "Diffing UObjects roundtripped via TPS...");
-		FDiffContext DiffCtx = {{ GUE.Schemas, Customs.Overlay }};
-		TUtf8StringBuilder<256> TpsDiffs;
-		TArray<int32> TpsDiffIdxs;
-		for (FInstance& Instance : Instances)
-		{
-			if (Algo::Find(SkipClasses, Instance.Orig->GetClass()->GetFName()))
-			{
-				continue;
-			}
-
-			if (DiffStructs(Instance.Orig, Instance.TPS, Instance.Id, /* in-out */ DiffCtx))
-			{
-				GTpsDiff.Instances.Add(&Instance);
-				PrintDiff(/* out */ TpsDiffs, GUE.Names, DiffCtx.Out);
-				
-				DiffCtx.Out.Reset();
-				TpsDiffs.Append(" in ");
-				TpsDiffs.Append(Instance.Orig->GetFullName());
-				TpsDiffs.Append("\n");	
-				TpsDiffIdxs.Add(&Instance - &Instances[0]);
-				
-				//FArchivedProperties Tmp;
-				//FPropertyWriter(/* out */ Tmp).WriteProperties(Instance.Orig, Instance.Base);
-				//FPropertyReader(Tmp).ReadProperties(Instance.UPS, Instance.Base);
-
-				//FDiffContext TmpCtx = DiffCtx;
-				//bool bOU = DiffStructs(Instance.Id, Instance.Orig, Instance.UPS, /* in-out */ TmpCtx);
-				//bool bOT = DiffStructs(Instance.Id, Instance.Orig, Instance.TPS, /* in-out */ TmpCtx);
-				//bool bOP = DiffStructs(Instance.Id, Instance.Orig, Instance.PP, /* in-out */ TmpCtx);
-				//TUtf8StringBuilder<256> TmpDiff;
-				//PrintDiff(/* out */ TmpDiff, TmpCtx.Out, GUE.Names);
-			}
-		}
-		GTpsDiff.Str = TpsDiffs.GetData();
-		UE_LOGFMT(LogPlainPropsUObject, Display, "Detected {Diffs} diffs in {Objs} UObjects saved in a {KB}KB value stream using TPS", TpsDiffIdxs.Num(), Instances.Num(), Tps.Data.NumBytes() / 1024);
-	}
-
-
-	if (EnumHasAnyFlags(Options, ERoundtrip::UPS))
-	{
-		FArchivedProperties Ups;
-		UE_LOGFMT(LogPlainPropsUObject, Display, "Saving UObjects to UPS archive...");
-		{
-			FPropertyWriter Archive(/* out */ Ups);
-			SaveArchive<true>(Archive, Instances);
-		}
-
-		UE_LOGFMT(LogPlainPropsUObject, Display, "Loading UObjects from UPS archive...");
-		{
-			FPropertyReader Archive(Ups);
-			LoadArchive<true>(Archive, Instances);
-		}
-
-		// Diff original vs UPS
-		UE_LOGFMT(LogPlainPropsUObject, Display, "Diffing UObjects roundtripped via UPS...");
-		FDiffContext DiffCtx = {{ GUE.Schemas, Customs.Overlay }};
-		TUtf8StringBuilder<256> UpsDiffs;
-		TArray<int32> UpsDiffIdxs;
-		for (FInstance& Instance : Instances)
-		{
-			if (Algo::Find(SkipClasses, Instance.Orig->GetClass()->GetFName()))
-			{
-				continue;
-			}
-
-			if (DiffStructs(Instance.Orig, Instance.UPS, Instance.Id, /* in-out */ DiffCtx))
-			{
-				GUpsDiff.Instances.Add(&Instance);
-				PrintDiff(/* out */ UpsDiffs, GUE.Names, DiffCtx.Out);
-				
-				DiffCtx.Out.Reset();
-				UpsDiffs.Append(" in ");
-				UpsDiffs.Append(Instance.Orig->GetFullName());
-				UpsDiffs.Append("\n");	
-				UpsDiffIdxs.Add(&Instance - &Instances[0]);
-			}
-		}
-		GUpsDiff.Str = UpsDiffs.GetData();
-		UE_LOGFMT(LogPlainPropsUObject, Display, "Detected {Diffs} diffs in {Objs} UObjects saved in a {KB}KB value stream using UPS", UpsDiffIdxs.Num(), Instances.Num(), Ups.Data.NumBytes() / 1024);
-	}
-
-	return NumPpDiffs;
-}
-
-////////////////////////////////////////////////////////////////////////////////
-
-int32 RoundtripViaPackages(TConstArrayView<UObject*> Objects, ERoundtrip Options)
-{
-	TArray<UPackage*> Packages;
-	TArray<UObject*> AllPackages;
-
-	// Collect all packages
-	if (Objects.Num())
-	{
-		for (UObject* Object : Objects)
-		{
-			AllPackages.Add(Object->GetPackage());
-		}
-	}
-	else
-	{
-		GetObjectsOfClass(UPackage::StaticClass(), AllPackages);
-	}
-
-	// Filter out in-memory and BP packages
-	for (UObject* PackageObj : AllPackages)
-	{
-		if (!PackageObj->HasAnyFlags(RF_WasLoaded))
-		{
-			continue;
-		}
-		UPackage* Package = Cast<UPackage>(PackageObj);
-		if (Package->HasAnyPackageFlags(PKG_CompiledIn))
-		{
-			continue;
-		}
-		bool bExclude = false;
-		ForEachObjectWithPackage(Package, [&bExclude](UObject* Object)
-		{
-			UClass* Class = Object->GetClass();
-			if (Class->HasAnyClassFlags(CLASS_CompiledFromBlueprint))
-			{
-				bExclude = true;
-				return false;
-			}
-			if (!IncludeClass(Class))
-			{
-				UE_LOGFMT(LogPlainPropsUObject, Log, " Excluding PlainProps for {Package}", Object->GetPackage()->GetName());
-				bExclude = true;
-				return false;
-			}
-			return true;
-		});
-		if (bExclude)
-		{
-			continue;
-		}
-		Packages.Add(Package);
-	}
-
-#if PP_WITH_ENGINE_HOOKS
-	// Todo consider TGuardValue
-	FSavePackageSettings& SavePackageSettings = FSavePackageSettings::GetDefaultSettings();
-	SavePackageSettings.ExternalWriteExportsFunc = &FPlainPropsLinkerSandbox::SavePlainProps;
-	SavePackageSettings.ExternalWriteFileFunc = &FPlainPropsLinkerSandbox::SavePackageToMemory;
-	GSerializePropertiesOverride = &FPlainPropsLinkerSandbox::LoadObject;
-	GLoadFromPackageFileOverride = &FPlainPropsLinkerSandbox::OpenPackageReader;
-#endif
-
-	UE_LOGFMT(LogPlainPropsUObject, Display, "Saving {SaveNum} packages...", Packages.Num());
-	for (UPackage* Package : Packages)
-	{
-		UObject* Asset = Package->FindAssetInPackage();
-		FSavePackageArgs SaveArgs;
-		SaveArgs.TopLevelFlags = RF_Standalone;
-		SaveArgs.SaveFlags = SAVE_Async; // Write progress to temp memory instead of a temp file
-		const FString& Extension = Package->ContainsMap() ? FPackageName::GetMapPackageExtension() : FPackageName::GetAssetPackageExtension();
-		FString TempFilename = FPaths::Combine(*FPaths::ProjectSavedDir(), TEXT("PlainProps"), Package->GetName() + Extension);
-		// FString ActualFilename = FPackageName::LongPackageNameToFilename(Asset->GetPackage()->GetName(), Extension);
-
-		UPackage::Save(Package, Asset, *TempFilename, SaveArgs);
-	}
-
-#if WITH_EDITOR
-	FPlainPropsLinkerSandbox::SavedPackages.Reserve(Packages.Num());
-	FPlainPropsLinkerSandbox::LoadBatches.Reserve(Packages.Num());
-
-	UE_LOGFMT(LogPlainPropsUObject, Display, "Mounting {SaveNum} packages...", Packages.Num());
-	for (UPackage* Package : Packages)
-	{
-		FPlainPropsLinkerSandbox::MountPackageFromMemory(Package->GetFName());
-	}
-
-	UE_LOGFMT(LogPlainPropsUObject, Display, "Loading {SaveNum} packages...", Packages.Num());
-	for (UPackage* Package : Packages)
-	{
-		FPlainPropsLinkerSandbox::LoadPackageFromMemory(Package->GetFName());
-	}
-#endif
-
-#if PP_WITH_ENGINE_HOOKS
-	SavePackageSettings.ExternalWriteExportsFunc = nullptr;
-	GSerializePropertiesOverride = nullptr;
-#endif
-
-	return 0;
-}
-
-struct FPlainPropsBindingScope
-{
-	FPlainPropsBindingScope()
-		: Transform()
-		, Guid()
-		, Color()
-		, LinearColor()
-		, FieldPath(GUE.Structs.FieldPath)
-		, Delegate(GUE.Structs.Delegate)
-		, InlineMulticast({ GUE.Structs.MulticastInlineDelegate, GUE.Structs.MulticastDelegate })
-		, VerseFunction(GUE.Structs.VerseFunction)
-		, DynamicallyTypedValue(GUE.Structs.DynamicallyTypedValue)
-		, ReferencePropertyValue(GUE.Structs.ReferencePropertyValue)
-	{
-	}
-
-	TScopedStructBinding<FTransform, FDefaultRuntime> Transform;
-	TScopedStructBinding<FGuid, FDefaultRuntime> Guid;
-	TScopedStructBinding<FColor, FDefaultRuntime> Color;
-	TScopedStructBinding<FLinearColor, FDefaultRuntime> LinearColor;
-	TScopedStructBinding<FFieldPath, FDefaultRuntime> FieldPath;
-	TScopedStructBinding<FScriptDelegate, FDefaultRuntime> Delegate;
-	// MulticastDelegate declaration is shared with MulticastSparseDelegate
-	TScopedStructBinding<FMulticastScriptDelegate, FDefaultRuntime> InlineMulticast;
-	// Verse
-	TScopedStructBinding<FVerseFunction, FDefaultRuntime> VerseFunction;
-	TScopedStructBinding<::UE::FDynamicallyTypedValue, FDefaultRuntime> DynamicallyTypedValue;
-	TScopedStructBinding<FReferencePropertyValue, FDefaultRuntime> ReferencePropertyValue;
-};
-
-} // namespace PlainProps::UE
-
-static TUniquePtr<PlainProps::UE::FPlainPropsBindingScope> BindingScopePtr = nullptr;
-static TUniquePtr<PlainProps::DbgVis::FIdScope> DebugIdScopePtr = nullptr;
-
-void FPlainPropsUObjectModule::StartupModule()
+void FPlainPropsUObjectModule::StartupModule()
 {
 	using namespace PlainProps;
 	using namespace PlainProps::UE;
-	BindingScopePtr = MakeUnique<FPlainPropsBindingScope>();
 	DebugIdScopePtr = MakeUnique<DbgVis::FIdScope>(GUE.Names, "SensName");
+	BindingScopePtr = MakeUnique<FPlainPropsBindingScope>();
 	InitBatchedProperties();
 }
 
diff --git a/Engine/Plugins/Experimental/PlainPropsUObject/Source/Private/PlainPropsUObjectBindingsInternal.h b/Engine/Plugins/Experimental/PlainPropsUObject/Source/Private/PlainPropsUObjectBindingsInternal.h
new file mode 100644
--- /dev/null
+++ b/Engine/Plugins/Experimental/PlainPropsUObject/Source/Private/PlainPropsUObjectBindingsInternal.h
@@ -0,0 +1,21 @@
+// Copyright Epic Games, Inc. All Rights Reserved.
+
+#pragma once
+
+#include "PlainPropsTypes.h"
+#include "Logging/LogMacros.h"
+
+DECLARE_LOG_CATEGORY_EXTERN(LogPlainPropsUObject, Log, All);
+
+class UField;
+class UStruct;
+
+namespace PlainProps::UE
+{
+
+FType			IndexType(const UField* Field);
+bool			ShouldBind(const UStruct* Struct);
+const UStruct*	SkipEmptyBases(const UStruct* In);
+
+} // namespace PlainProps::UE
+
