PlainProps: Cleanup unused PP_WITH_ENGINE_HOOKS define in PlainPropsUObjectBindinds.cpp.

#rnx

diff --git a/Engine/Plugins/Experimental/PlainPropsUObject/Source/Private/PlainPropsUObjectBindings.cpp b/Engine/Plugins/Experimental/PlainPropsUObject/Source/Private/PlainPropsUObjectBindings.cpp
--- a/Engine/Plugins/Experimental/PlainPropsUObject/Source/Private/PlainPropsUObjectBindings.cpp
+++ b/Engine/Plugins/Experimental/PlainPropsUObject/Source/Private/PlainPropsUObjectBindings.cpp
@@ -32,8 +32,6 @@
 #include "UObject/VerseValueProperty.h"
 
 
-#define PP_WITH_ENGINE_HOOKS 0
-
 using FUnicastScriptDelegate = TScriptDelegate<FNotThreadSafeNotCheckedDelegateMode>;
 using FMulticastInvocationList = TArray<FUnicastScriptDelegate>;
 using FMulticastInvocationView = TConstArrayView<FUnicastScriptDelegate>;
