PlainProps fix FTransform and some other core types registering a default value now that they're properly custom-bound

#rnx

diff --git a/Engine/Plugins/Experimental/PlainPropsUObject/Source/Private/PlainPropsUObjectBindings.cpp b/Engine/Plugins/Experimental/PlainPropsUObject/Source/Private/PlainPropsUObjectBindings.cpp
--- a/Engine/Plugins/Experimental/PlainPropsUObject/Source/Private/PlainPropsUObjectBindings.cpp
+++ b/Engine/Plugins/Experimental/PlainPropsUObject/Source/Private/PlainPropsUObjectBindings.cpp
@@ -3533,6 +3533,8 @@ static void InitBatchedProperties()
 
 struct FPlainPropsBindingScope
 {
+	static const FTransform DefaultTransform;
+
 	FPlainPropsBindingScope()
 		: Transform()
 		, Guid()
@@ -3545,6 +3547,18 @@ struct FPlainPropsBindingScope
 		, DynamicallyTypedValue(GUE.Structs.DynamicallyTypedValue)
 		, ReferencePropertyValue(GUE.Structs.ReferencePropertyValue)
 	{
+		GUE.Defaults.BindStatic(Transform.BindId, &DefaultTransform);
+		GUE.Defaults.BindZeroes<FGuid>(Guid.BindId);
+		GUE.Defaults.BindZeroes<FColor>(Color.BindId);
+		GUE.Defaults.BindZeroes<FLinearColor>(LinearColor.BindId);
+	}
+
+	~FPlainPropsBindingScope()
+	{
+		GUE.Defaults.Drop(Transform.BindId);
+		GUE.Defaults.Drop(Guid.BindId);
+		GUE.Defaults.Drop(Color.BindId);
+		GUE.Defaults.Drop(LinearColor.BindId);
 	}
 
 	TScopedStructBinding<FTransform, FDefaultRuntime> Transform;
diff --git a/Engine/Plugins/Experimental/PlainPropsUObject/Source/Public/PlainPropsUObjectRuntime.h b/Engine/Plugins/Experimental/PlainPropsUObject/Source/Public/PlainPropsUObjectRuntime.h
--- a/Engine/Plugins/Experimental/PlainPropsUObject/Source/Public/PlainPropsUObjectRuntime.h
+++ b/Engine/Plugins/Experimental/PlainPropsUObject/Source/Public/PlainPropsUObjectRuntime.h
@@ -31,6 +31,9 @@ public:
 	void									Drop(FBindId Id);
 	virtual const void*						Get(FBindId Id) override;
 
+	template<class T>
+	void									BindZeroes(FBindId Id) { BindZeroes(Id, sizeof(T), alignof(T)); }
+
 private:
 	FMutableMemoryView						Zeroes;
 	TBitArray<>								Instanced;
