PlainProps fix compile error

#rnx

diff --git a/Engine/Plugins/Experimental/PlainPropsUObject/Source/Private/PlainPropsUObjectBindings.cpp b/Engine/Plugins/Experimental/PlainPropsUObject/Source/Private/PlainPropsUObjectBindings.cpp
--- a/Engine/Plugins/Experimental/PlainPropsUObject/Source/Private/PlainPropsUObjectBindings.cpp
+++ b/Engine/Plugins/Experimental/PlainPropsUObject/Source/Private/PlainPropsUObjectBindings.cpp
@@ -3533,7 +3533,7 @@ static void InitBatchedProperties()
 
 struct FPlainPropsBindingScope
 {
-	static const FTransform DefaultTransform;
+	inline static const FTransform DefaultTransform;
 
 	FPlainPropsBindingScope()
 		: Transform()
