[Core] Add support to compile switch between using sparse or compact sets as the default set container

#tests Using the editor, running PIE and game client, making builds

diff --git a/Engine/Plugins/Experimental/PlainPropsUObject/Source/Private/PlainPropsUObjectBindings.cpp b/Engine/Plugins/Experimental/PlainPropsUObject/Source/Private/PlainPropsUObjectBindings.cpp
--- a/Engine/Plugins/Experimental/PlainPropsUObject/Source/Private/PlainPropsUObjectBindings.cpp
+++ b/Engine/Plugins/Experimental/PlainPropsUObject/Source/Private/PlainPropsUObjectBindings.cpp
@@ -1638,31 +1638,10 @@ using TPropertyRangeSerializer = typename TSelectRangeSerializer<Kind>::Type;
 
 ////////////////////////////////////////////////////////////////////////////////////////////////
 
+#if !UE_USE_COMPACT_SET_AS_DEFAULT
 inline FScriptSparseArray& AsSparseArray(FScriptSet& In) { return reinterpret_cast<FScriptSparseArray&>(In); }
 inline FScriptSparseArray& AsSparseArray(FScriptMap& In) { return reinterpret_cast<FScriptSparseArray&>(In); }
 
-template<typename ScriptSet>
-inline bool IsCompact(const ScriptSet& Set)
-{
-	return Set.NumUnchecked() == Set.GetMaxIndex();
-}
-
-// There's no TScriptSparseArray::SetNumUninitialized() (yet), 
-// reserve using Empty() and add items one by one instead
-template<class ScriptType, class LayoutType>
-uint8* SetNumUninitialized(ScriptType& Dst, const LayoutType& Layout, uint64 Num)
-{
-	check(Dst.IsEmpty());
-	Dst.Empty(static_cast<int32>(Num), Layout);
-	for (uint64 Idx = 0; Idx < Num; ++Idx)
-	{
-		Dst.AddUninitialized(Layout);
-	}
-	check(IsCompact(Dst));
-
-	return static_cast<uint8*>(Dst.GetData(0, Layout));
-}
-
 // @pre Elems.Num() > 0
 inline FExistingItemSlice GetContiguousSlice(int32 Idx, const FScriptSparseArray& Elems, const uint8* Data, SIZE_T Stride)
 {
@@ -1704,6 +1683,29 @@ inline void ReadSparseItems(FExistingItems& Dst, const FScriptSparseArray& Src,
 		Dst.Slice = GetContiguousSlice(0, Src, Data, Layout.Size);
 	}
 }
+#endif
+
+template<typename ScriptSet>
+inline bool IsCompact(const ScriptSet& Set)
+{
+	return Set.NumUnchecked() == Set.GetMaxIndex();
+}
+
+// There's no TScriptSparseArray::SetNumUninitialized() (yet), 
+// reserve using Empty() and add items one by one instead
+template<class ScriptType, class LayoutType>
+uint8* SetNumUninitialized(ScriptType& Dst, const LayoutType& Layout, uint64 Num)
+{
+	check(Dst.IsEmpty());
+	Dst.Empty(static_cast<int32>(Num), Layout);
+	for (uint64 Idx = 0; Idx < Num; ++Idx)
+	{
+		Dst.AddUninitialized(Layout);
+	}
+	check(IsCompact(Dst));
+
+	return static_cast<uint8*>(Dst.GetData(0, Layout));
+}
 
 ////////////////////////////////////////////////////////////////////////////////////////////////
 
@@ -1948,7 +1950,11 @@ struct TSetPropertyBinding : IItemRangeBinding, ICustomBinding
 		return Declare({Id, NoId, 0, EMemberPresence::AllowSparse, Members, {RangeSpec, RangeSpec, RangeSpec}});
 	}
 
-	inline SIZE_T GetStride() const requires (bLeaves)	{ return sizeof(TSetElement<LeafType>); }
+#if UE_USE_COMPACT_SET_AS_DEFAULT
+	inline SIZE_T GetStride() const requires (bLeaves)	{ return sizeof(LeafType); }
+#else
+	inline SIZE_T GetStride() const requires (bLeaves)	{ return sizeof(TSparseSetElement<LeafType>); }
+#endif
 	inline SIZE_T GetStride() const						{ return Layout.Size; }
 
 	inline int32 FindIndex(const FScriptSet& Set, const void* Elem) const
@@ -1961,7 +1967,11 @@ struct TSetPropertyBinding : IItemRangeBinding, ICustomBinding
 		if (int32 Idx = FindIndex(Set, Elem); Idx != INDEX_NONE)
 		{
 			DestroyElem(Set, Idx);
+#if UE_USE_COMPACT_SET_AS_DEFAULT
+			Set.RemoveAt(Idx, Layout, MakeHashFn(Inner));
+#else
 			Set.RemoveAt(Idx, Layout);
+#endif
 		}
 	}
 	
@@ -1975,11 +1985,6 @@ struct TSetPropertyBinding : IItemRangeBinding, ICustomBinding
 		return HasItem(Set, Elem);
 	}
 
-	inline void Rehash(FScriptSet& Set) const
-	{
-		Set.Rehash(Layout, MakeHashFn(Inner));
-	}
-
 	inline void DestroyElem(FScriptSet& Set) const requires (bLeaves) {}
 	inline void DestroyElem(FScriptSet& Set, int32 Idx) const
 	{
@@ -2010,7 +2015,12 @@ struct TSetPropertyBinding : IItemRangeBinding, ICustomBinding
 	
 	virtual void ReadItems(FSaveRangeContext& Ctx) const override
 	{
+#if UE_USE_COMPACT_SET_AS_DEFAULT
+		const FScriptSet& Set = Ctx.Request.GetRange<FScriptSet>();
+		Ctx.Items.SetAll(Set.GetData(0, Layout), Set.Num(), Layout.Size);
+#else
 		ReadSparseItems(/* out */ Ctx.Items, Ctx.Request.GetRange<FScriptSparseArray>(), Layout.SparseArrayLayout);
+#endif
 	}
 
 	virtual void MakeItems(FLoadRangeContext& Ctx) const override
@@ -2033,7 +2043,7 @@ struct TSetPropertyBinding : IItemRangeBinding, ICustomBinding
 		else
 		{
 			check(Ctx.Request.IsFinalCall());
-			Rehash(Set);
+			Set.Rehash(Layout, MakeHashFn(Inner));
 		}
 	}
 
@@ -2086,38 +2096,42 @@ struct TSetPropertyBinding : IItemRangeBinding, ICustomBinding
 			}
 		}
 
-		Rehash(Dst);
+		Dst.Rehash(Layout, MakeHashFn(Inner));
 	}
 
 	// Load leaves into non-empty set
 	inline void InsertNonEmpty(FScriptSet& Dst, FRangeLoadView Src) const requires(bLeaves)
 	{	
+		auto HashFn = MakeHashFn(Inner);
+
 		for (LeafType Item : CastAs<LeafType>(Src.AsLeaves()))
 		{
 			if (!HasItem(Dst, &Item))
 			{
 				 void* Elem = Dst.GetData(Dst.AddUninitialized(Layout), Layout);
+				 Dst.CommitLastUninitialized(Layout, HashFn);
 				 *static_cast<LeafType*>(Elem) = Item;
 			}
 		}
 
-		Rehash(Dst);
+		Dst.CommitAllUninitialized(Layout, HashFn);
 	}
 
-	inline void* AddItem(FScriptSparseArray& Dst, int32& OutIdx) const
+	inline void* AddItem(FScriptSet& Dst, int32& OutIdx) const
 	{
-		OutIdx = Dst.AddUninitialized(Layout.SparseArrayLayout);
-		void* Out = Dst.GetData(OutIdx, Layout.SparseArrayLayout);
+		OutIdx = Dst.AddUninitialized(Layout);
+		void* Out = Dst.GetData(OutIdx, Layout);
 		Inner.InitItem(Out);
 		return Out;
 	}
 
 	// Load structs or ranges into non-empty set
-	inline void InsertNonEmpty(FScriptSet& DstSet, FRangeLoadView Src) const
+	inline void InsertNonEmpty(FScriptSet& Dst, FRangeLoadView Src) const
 	{
+		auto HashFn = MakeHashFn(Inner);
+
 		// Written to avoid FProperty::CopyCompleteValue_InContainer dependency
 		// Items are loaded directly into sparse array and then removed if a duplicate existed
-		FScriptSparseArray& Dst = AsSparseArray(DstSet);
 		const int32 OldNum = Dst.NumUnchecked();
 		int32 TmpIdx;
 		void* Tmp = AddItem(Dst, /* out */ TmpIdx);
@@ -2127,7 +2141,11 @@ struct TSetPropertyBinding : IItemRangeBinding, ICustomBinding
 			for (FRangeLoadView Item : Src.AsRanges())
 			{
 				LoadRange(Tmp, Item, InnerBindings);
-				Tmp = HasItem(DstSet, Tmp) ? Tmp : AddItem(Dst, /* out */ TmpIdx);
+				if (!HasItem(Dst, Tmp))
+				{
+					Dst.CommitLastUninitialized(Layout, HashFn);
+					Tmp = AddItem(Dst, /* out */ TmpIdx);
+				}
 			}
 		}
 		else
@@ -2135,16 +2153,20 @@ struct TSetPropertyBinding : IItemRangeBinding, ICustomBinding
 			for (FStructLoadView Item : Src.AsStructs())
 			{
 				LoadStruct(Tmp, Item);
-				Tmp = HasItem(DstSet, Tmp) ? Tmp : AddItem(Dst, /* out */ TmpIdx);
+				if (!HasItem(Dst, Tmp))
+				{
+					Dst.CommitLastUninitialized(Layout, HashFn);
+					Tmp = AddItem(Dst, /* out */ TmpIdx);
+				}
 			}
 		}
 
 		Inner.DestroyItem(Tmp);
-		Dst.RemoveAtUninitialized(Layout.SparseArrayLayout, TmpIdx, 1);
+		Dst.RemoveAtUninitialized(Layout, TmpIdx);
 
 		if (Dst.NumUnchecked() != OldNum)
 		{
-			Rehash(DstSet);
+			Dst.CommitAllUninitialized(Layout, HashFn);
 		}
 	}
 
@@ -2188,7 +2210,7 @@ struct TSetPropertyBinding : IItemRangeBinding, ICustomBinding
 	{
 		checkSlow(Idx < Dst.GetMaxIndex());
 		checkSlow(!Dst.IsValidIndex(Idx));
-		void* Elem = AsSparseArray(Dst).GetData(Idx, Layout.SparseArrayLayout);
+		void* Elem = Dst.GetData(Idx, Layout);
 		Inner.DestroyItem(Elem);
 	}
 };
@@ -2310,7 +2332,12 @@ struct TMapPropertyItemBinding : IItemRangeBinding
 
 	virtual void ReadItems(FSaveRangeContext& Ctx) const override
 	{
+#if UE_USE_COMPACT_SET_AS_DEFAULT
+		const FScriptMap& Set = Ctx.Request.GetRange<FScriptMap>();
+		Ctx.Items.SetAll(Set.GetData(0, Layout), Set.Num(), Layout.SetLayout.Size);
+#else
 		ReadSparseItems(/* out */ Ctx.Items, Ctx.Request.GetRange<FScriptSparseArray>(), Layout.SetLayout.SparseArrayLayout);
+#endif
 	}
 
 	virtual void MakeItems(FLoadRangeContext& Ctx) const override
@@ -2406,12 +2433,12 @@ struct TMapPropertyCustomBinding : TMapPropertyItemBinding<FScriptMap, KeyKind,
 		return false;
 	}
 
-	inline uint8* AddPair(FScriptSparseArray& Dst, int32& OutIdx, int32 ValueOffset) const
+	inline uint8* AddPair(FScriptMap& Dst, int32& OutIdx) const
 	{
-		OutIdx = Dst.AddUninitialized(Layout.SetLayout.SparseArrayLayout);
-		uint8* Out = static_cast<uint8*>(Dst.GetData(OutIdx, Layout.SetLayout.SparseArrayLayout));
+		OutIdx = Dst.AddUninitialized(Layout);
+		uint8* Out = static_cast<uint8*>(Dst.GetData(OutIdx, Layout));
 		InnerKey.InitItem(Out);
-		InnerValue.InitItem(Out + ValueOffset);
+		InnerValue.InitItem(Out + Layout.ValueOffset);
 		return Out;
 	}
 
@@ -2430,7 +2457,11 @@ struct TMapPropertyCustomBinding : TMapPropertyItemBinding<FScriptMap, KeyKind,
 		if (int32 Idx = FindKey(Map, Key); Idx != INDEX_NONE)
 		{
 			DestroyPair(Map, Idx);
+#if UE_USE_COMPACT_SET_AS_DEFAULT
+			Map.RemoveAt(Idx, Layout, MakeHashFn(InnerKey));
+#else
 			Map.RemoveAt(Idx, Layout);
+#endif
 		}
 	}
 
@@ -2474,14 +2505,15 @@ struct TMapPropertyCustomBinding : TMapPropertyItemBinding<FScriptMap, KeyKind,
 	// Load structs or ranges into non-empty map
 	inline void InsertNonEmpty(FScriptMap& Dst, FRangeLoadView Src) const
 	{
+		auto HashFn = MakeHashFn(InnerKey);
+
 		// Written to avoid FProperty::CopyCompleteValue_InContainer dependency
 		// Items are loaded directly into sparse array and then removed if a duplicate existed
 
-		FScriptSparseArray& DstArray = AsSparseArray(Dst);
-		const int32 OldNum = DstArray.NumUnchecked();
+		const int32 OldNum = Dst.NumUnchecked();
 		const int32 ValueOffset = Layout.ValueOffset;
 		int32 TmpIdx;
-		uint8* Tmp = AddPair(DstArray, /* out*/ TmpIdx, ValueOffset);
+		uint8* Tmp = AddPair(Dst, /* out*/ TmpIdx);
 
 		FOptionalSchemaId InnerLoadIds[2];
 		FStructRangeLoadView Structs = Src.AsStructs();
@@ -2495,24 +2527,25 @@ struct TMapPropertyCustomBinding : TMapPropertyItemBinding<FScriptMap, KeyKind,
 			if (int32 Idx = FindKey(Dst, Tmp); Idx != INDEX_NONE)
 			{
 				// Load value into existing pair
-				uint8* Pair = static_cast<uint8*>(DstArray.GetData(Idx, Layout.SetLayout.SparseArrayLayout));
+				uint8* Pair = static_cast<uint8*>(Dst.GetData(Idx, Layout));
 				ValueRange.LoadItem(Pair + ValueOffset,	/* in-out */ Struct.Values, /* in-out */ Bits, InnerLoadIds[1], Struct.Schema.Batch);
 			}
 			else
 			{
 				// Load value into tmp pair and add new temporary
 				ValueRange.LoadItem(Tmp + ValueOffset,	/* in-out */ Struct.Values, /* in-out */ Bits, InnerLoadIds[1], Struct.Schema.Batch);
-				Tmp = AddPair(DstArray, /* out*/ TmpIdx, ValueOffset);
+				Dst.CommitLastUninitialized(Layout, HashFn);
+				Tmp = AddPair(Dst, /* out*/ TmpIdx);
 			}
 		}
 
 		InnerKey.DestroyItem(Tmp);
 		InnerValue.DestroyItem(Tmp + ValueOffset);
-		DstArray.RemoveAtUninitialized(Layout.SetLayout.SparseArrayLayout, TmpIdx, 1);
+		Dst.RemoveAtUninitialized(Layout, TmpIdx);
 
-		if (DstArray.NumUnchecked() != OldNum)
+		if (Dst.NumUnchecked() != OldNum)
 		{
-			Rehash(Dst);
+			Dst.CommitAllUninitialized(Layout, HashFn);
 		}
 	}
 
@@ -3476,7 +3509,7 @@ static void BindStruct(FBindId Id, const UStruct* Struct)
 	BindMembers(/* out */ Binder, Struct);
 
 	GUE.Schemas.BindStruct(Id, Binder.GetMembers(), Binder.Declare());
-		
+	
 	// Don't bind CDO defaults, object defaults are passed in from top and objects aren't owned by containers 
 	if (Struct->HasAnyCastFlags(CASTCLASS_UScriptStruct))
 	{
diff --git a/Engine/Source/Runtime/Core/Public/Containers/Array.h b/Engine/Source/Runtime/Core/Public/Containers/Array.h
--- a/Engine/Source/Runtime/Core/Public/Containers/Array.h
+++ b/Engine/Source/Runtime/Core/Public/Containers/Array.h
@@ -216,6 +216,18 @@ UE_NODEBUG [[nodiscard]] UE_FORCEINLINE_HINT TIndexedContainerIterator<Container
 		{
 		}
 
+		UE_NODEBUG [[nodiscard]] UE_FORCEINLINE_HINT ElementType* operator->() const
+		{
+			if constexpr (bReverse)
+			{
+				return Ptr - 1;
+			}
+			else
+			{
+				return Ptr;
+			}
+		}
+
 		UE_NODEBUG [[nodiscard]] UE_FORCEINLINE_HINT ElementType& operator*() const
 		{
 			if constexpr (bReverse)
@@ -265,6 +277,12 @@ UE_NODEBUG [[nodiscard]] UE_FORCEINLINE_HINT TIndexedContainerIterator<Container
 			return Ptr != Rhs.Ptr;
 		}
 
+		UE_NODEBUG [[nodiscard]] FORCEINLINE bool operator==(const TCheckedPointerIterator& Rhs) const
+		{
+			return !(*this != Rhs);
+		}
+
+
 	private:
 		ElementType*    Ptr;
 		const SizeType& CurrentNum;
diff --git a/Engine/Source/Runtime/CoreUObject/Private/UObject/GarbageCollectionSchema.cpp b/Engine/Source/Runtime/CoreUObject/Private/UObject/GarbageCollectionSchema.cpp
--- a/Engine/Source/Runtime/CoreUObject/Private/UObject/GarbageCollectionSchema.cpp
+++ b/Engine/Source/Runtime/CoreUObject/Private/UObject/GarbageCollectionSchema.cpp
@@ -14,7 +14,7 @@ namespace UE::GC {
 //////////////////////////////////////////////////////////////////////////
 
 static constexpr uint32 ToBit(EMemberType Type)		{ return 1u << static_cast<uint32>(Type); }
-static constexpr uint32 InnerSchemaTypes =			ToBit(EMemberType::StructArray) | ToBit(EMemberType::SparseStructArray) | ToBit(EMemberType::FreezableStructArray) | ToBit(EMemberType::Optional);
+static constexpr uint32 InnerSchemaTypes =			ToBit(EMemberType::StructArray) | ToBit(EMemberType::StructSet) | ToBit(EMemberType::FreezableStructArray) | ToBit(EMemberType::Optional);
 static constexpr uint32 AROTypes =					ToBit(EMemberType::ARO) | ToBit(EMemberType::SlowARO) | ToBit(EMemberType::MemberARO);
 static constexpr uint32 AROPointerTypes =			ToBit(EMemberType::ARO) | ToBit(EMemberType::MemberARO);
 static constexpr uint32 ExtraWordTypes =			InnerSchemaTypes | AROPointerTypes | ToBit(EMemberType::StridedArray);
@@ -46,7 +46,7 @@ FName ToName(EMemberType Type)
 			Names[(uint8)EMemberType::ReferenceArray] =					"ReferenceArray";
 			Names[(uint8)EMemberType::StructArray] =					"StructArray";
 			Names[(uint8)EMemberType::StridedArray] =					"StridedArray";
-			Names[(uint8)EMemberType::SparseStructArray] =				"SparseStructArray";
+			Names[(uint8)EMemberType::StructSet] =						"StructSet";
 			Names[(uint8)EMemberType::FieldPath] =						"FieldPath";
 			Names[(uint8)EMemberType::FieldPathArray] =					"FieldPathArray";
 			Names[(uint8)EMemberType::FreezableReferenceArray] =		"FreezableReferenceArray";
@@ -649,7 +649,7 @@ struct FStructMemberStats
 		{ 
 			case EMemberType::StructArray:			VisitArray(*(const FScriptArray*)Container, Visitor);
 			break;
-			case EMemberType::SparseStructArray:	VisitSparseArray(*(const FScriptSparseArray*)Container, Visitor);
+			case EMemberType::StructSet:			VisitSet(*(const FScriptSet*)Container, Visitor);
 			break;
 			case EMemberType::FreezableStructArray:	VisitArray(*(const FFreezableScriptArray*)Container, Visitor);
 			break;
@@ -678,13 +678,13 @@ private:
 	}
 
 	template<typename Fn>
-	void VisitSparseArray(const FScriptSparseArray& Array, Fn Visitor) const
+	void VisitSet(const FScriptSet& Set, Fn Visitor) const
 	{
-		const uint8* Data = Private::GetSparseData(Array);
+		const uint8* Data = Private::GetSetData(Set);
 		const uint32 Stride = Schema.GetStructStride();
-		for (int32 Idx = 0, Num = Array.Num(); Idx < Num; ++Idx)
+		for (int32 Idx = 0, Num = Set.Num(); Idx < Num; ++Idx)
 		{
-			if (Array.IsValidIndex(Idx))
+			if (Set.IsValidIndex(Idx))
 			{
 				Visitor(Data + Idx * Stride);
 			}
