PlainProps: Fix json parsing error in PlainPropsUObject.uplugin.

#rnx

diff --git a/Engine/Plugins/Experimental/PlainPropsUObject/PlainPropsUObject.uplugin b/Engine/Plugins/Experimental/PlainPropsUObject/PlainPropsUObject.uplugin
--- a/Engine/Plugins/Experimental/PlainPropsUObject/PlainPropsUObject.uplugin
+++ b/Engine/Plugins/Experimental/PlainPropsUObject/PlainPropsUObject.uplugin
@@ -15,7 +15,7 @@
 			"Name": "PlainPropsUObject",
 			"Type": "Runtime",
 			"LoadingPhase": "PreLoadingScreen",
-			"PlatformAllowList": [ "Win64" ],
+			"PlatformAllowList": [ "Win64" ]
 		}
 	],
 	"Plugins": [
