PlainProps: Save and load (roundtrip) source packages to and from memory using FLinkerLoad and FLinkerSave.
1. Load all packages from disk.
2. Save all loaded packages with both TPS and PP.
3. Load all saved packages from step 2 into TPS and PP clones with FName postfix numbers 8848 and 8849.
4. Diff all objects found in the loaded clones from step 3.
Add custom bindings for FAnimationAttributeIdentifier, FAttributeCurveBinding, FNiagaraTypeDefinition and FNiagaraTypeDefinitionHandleBinding.
Add tailored uobject diff logic.

Roundtrip testing: -bind=source -package -loadmaps -allowcommandletaudio [-text]
-text will dump all saved packages to yaml files instead of doing the roundtripping.

#rnx

diff --git a/Engine/Plugins/Experimental/PlainProps/Source/Internal/PlainPropsInternalPrivateMemberPtr.h b/Engine/Plugins/Experimental/PlainProps/Source/Internal/PlainPropsInternalPrivateMemberPtr.h
new file mode 100644
--- /dev/null
+++ b/Engine/Plugins/Experimental/PlainProps/Source/Internal/PlainPropsInternalPrivateMemberPtr.h
@@ -0,0 +1,29 @@
+// Copyright Epic Games, Inc. All Rights Reserved.
+
+#pragma once
+
+#include "HAL/PreprocessorHelpers.h"
+#include "Templates/Identity.h"
+
+#define PP_DEFINE_PRIVATE_MEMBER_NAMESPACE_BEGIN(Namespace) \
+	namespace Namespace \
+	{ \
+		template <auto Storage, auto PtrToMember> \
+		struct TPrivateAccess \
+		{ \
+			TPrivateAccess() \
+			{ \
+				*Storage = PtrToMember; \
+			} \
+			static TPrivateAccess Instance; \
+		}; \
+		template <auto Storage, auto PtrToMember> \
+		TPrivateAccess<Storage, PtrToMember> TPrivateAccess<Storage, PtrToMember>::Instance;
+
+#define PP_DEFINE_PRIVATE_MEMBER_NAMESPACE_END() \
+	}
+
+#define PP_DEFINE_PRIVATE_MEMBER_PTR(Class, Member, Type) \
+	TIdentity_T<UE_REMOVE_OPTIONAL_PARENS(Type)> UE_REMOVE_OPTIONAL_PARENS(Class)::* _##Member; \
+	template struct TPrivateAccess<&_##Member, &UE_REMOVE_OPTIONAL_PARENS(Class)::Member>
+
diff --git a/Engine/Plugins/Experimental/PlainProps/Source/Public/PlainPropsBindCtti.h b/Engine/Plugins/Experimental/PlainProps/Source/Public/PlainPropsBindCtti.h
--- a/Engine/Plugins/Experimental/PlainProps/Source/Public/PlainPropsBindCtti.h
+++ b/Engine/Plugins/Experimental/PlainProps/Source/Public/PlainPropsBindCtti.h
@@ -49,17 +49,25 @@ struct FRequireAll
 ////////////////////////////////////////////////////////////////////////////////////////////////
 
 template<class Ids, class Typename>
-FScopeId IndexNamespaceId() requires (Typename::Namespace.size() > 0)
+FScopeId GetOuterScope()
 {
-	// Opt: Make cached GetNamespaceId(), either via new namespace CTTI types (maybe PP_REFLECT_NAMESPACE)
+	return NoId;
+}
+
+template<class Ids, class Typename>
+FScopeId IndexScope()
+{
+	// Opt: Make cached GetScopeId(), either via new namespace CTTI types (maybe PP_REFLECT_NAMESPACE)
 	//		or some compile time string template parameters, perhaps a variadic template taking any number of chars
-	return Ids::IndexScope(ToAnsiView(Typename::Namespace));
+	return GetOuterScope<Ids, Typename>();
 }
 
 template<class Ids, class Typename>
-FScopeId IndexNamespaceId()
+FScopeId IndexScope() requires (Typename::Namespace.size() > 0)
 {
-	return NoId;
+	FScopeId Outer = GetOuterScope<Ids, Typename>();
+	FFlatScopeId Namespace = Ids::IndexScope(ToAnsiView(Typename::Namespace));
+	return Outer ? Ids::GetIndexer().NestFlatScope(Outer, Namespace) : FScopeId(Namespace);
 }
 
 template<ETypename Kind, class Typename>
@@ -78,8 +86,8 @@ constexpr std::string_view SelectStructName()
 template<class Ids, ETypename Kind, class Typename>
 FType IndexStructName()
 {
-	FType BaseName = { IndexNamespaceId<Ids, Typename>(), 
-						 FTypenameId(Ids::IndexTypename(ToAnsiView(SelectStructName<Kind, Typename>()))) };
+	FType BaseName = { IndexScope<Ids, Typename>(), 
+					   FTypenameId(Ids::IndexTypename(ToAnsiView(SelectStructName<Kind, Typename>()))) };
 	
 	if constexpr (ParametricName<Typename>)
 	{
@@ -145,23 +153,18 @@ FDualStructId GetStructDualId()
 	return Id;
 }
 
-template<class Ids, typename Ctti>
-FType IndexCttiName()
+template<class Ids, Enumeration Enum>
+FType IndexEnumName()
 {
-	FTypenameId Name{Ids::IndexTypename(Ctti::Name)};
-	FScopeId Namespace = NoId;
-	if constexpr (Ctti::Namespace[0] != '\0')
-	{
-		Namespace = Ids::IndexScope(ToAnsiView(Ctti::Namespace));
-	}
-	return { Namespace, Name };
+	using Typename = TTypename<Enum>;
+	return { IndexScope<Ids, Typename>(), FTypenameId(Ids::IndexTypename(ToAnsiView(Typename::DeclName))) };
 }
 
 // Cached by function static
 template<class Ids, Enumeration Enum>
 FEnumId GetEnumId()
 {
-	static FEnumId Id = Ids::IndexEnum(IndexCttiName<Ids, CttiOf<Enum>>());
+	static FEnumId Id = Ids::IndexEnum(IndexEnumName<Ids, Enum>());
 	return Id;
 }
 
@@ -205,7 +208,7 @@ FType IndexParameterName()
 		else
 		{
 			using Typename = TTypename<T>;
-			FType RangeBindName = { IndexNamespaceId<Ids, Typename>(), FTypenameId(Ids::IndexTypename(ToAnsiView(RangeBinding::BindName))) };
+			FType RangeBindName = { IndexScope<Ids, Typename>(), FTypenameId(Ids::IndexTypename(ToAnsiView(RangeBinding::BindName))) };
 			return Ids::GetIndexer().MakeParametricType(RangeBindName, {ItemParam, SizeParam});
 		}
 	}
@@ -376,6 +379,16 @@ struct TScopedStructBinding<T, Runtime, void> : FBothStructId
 	}
 };
 
+template<Enumeration Enum, EEnumMode Mode, class Runtime>
+struct TScopedEnumDeclaration
+{
+	using Ids = typename Runtime::Ids;
+
+	FEnumId Id;
+	TScopedEnumDeclaration() : Id(DeclareNativeEnum<Enum, Ids>(Runtime::GetEnums(), Mode)) {}
+	~TScopedEnumDeclaration() { Runtime::GetEnums().Drop(Id); }
+};
+
 template<class Struct, typename CustomBinding, class Runtime>
 FBothStructId BindCustomStructOnce()
 {
@@ -535,12 +548,13 @@ FMemberBinding BindMember(uint64 Offset, FMemberSpec& OutSpec)
 
 ////////////////////////////////////////////////////////////////////////////////////////////////
 
-template<class Ctti, class Ids>
+template<Enumeration Enum, class Ids>
 FEnumId DeclareNativeEnum(FEnumDeclarations& Out, EEnumMode Mode)
 {
+	using Ctti = CttiOf<Enum>;
 	using UnderlyingType = std::underlying_type_t<typename Ctti::Type>;
 
-	FType Type = IndexCttiName<Ids, Ctti>();
+	FType Type = IndexEnumName<Ids, Enum>();
 	FEnumId Id = Ids::IndexEnum(Type);
 	FEnumerator Enumerators[Ctti::NumEnumerators];
 	for (FEnumerator& Enumerator : Enumerators)
@@ -612,4 +626,4 @@ TStaticArray<FAuxMember, Ctti::NumVars> BindAuxMembers()
 	return Out;
 }
 
-} // namespace PlainProps
\ No newline at end of file
+} // namespace PlainProps
diff --git a/Engine/Plugins/Experimental/PlainProps/Source/Public/PlainPropsCtti.h b/Engine/Plugins/Experimental/PlainProps/Source/Public/PlainPropsCtti.h
--- a/Engine/Plugins/Experimental/PlainProps/Source/Public/PlainPropsCtti.h
+++ b/Engine/Plugins/Experimental/PlainProps/Source/Public/PlainPropsCtti.h
@@ -31,7 +31,7 @@ template<typename T> using CttiOf = typename TCttiOf<T>::Type;
 #define _PP_REFLECT_THINGS(N, MACRO, ...)				_PP_EXPAND(_PP_REFLECT_##N(N, MACRO, __VA_ARGS__))
 
 #define _PP_REFLECT_ENUMERATOR(N, NS, T, C)				{ #C, Type::C }, 
-#define _PP_REFLECT_ENUM_INNER(N, NS, T, ES)			struct T##_Ctti { inline static constexpr char Name[] = #T; using Type = NS :: T; static constexpr int NumEnumerators = N; static constexpr struct { const char* Name; Type Constant; } Enumerators[] = { ES }; inline static constexpr char Namespace[] = #NS;  }; T##_Ctti CttiOfPtr(T*);
+#define _PP_REFLECT_ENUM_INNER(N, NS, T, ES)			struct T##_Ctti { inline static constexpr char Name[] = #T; using Type = NS :: T; static constexpr int NumEnumerators = N; static constexpr struct { const char* Name; Type Constant; } Enumerators[] = { ES }; inline static constexpr char Namespace[] = #NS;  }; T##_Ctti CttiOfPtr(NS :: T*);
 #define _PP_REFLECT_ENUM(N, NS, T, ...)					_PP_REFLECT_ENUM_INNER(N, NS, T, _PP_EXPAND(_PP_REFLECT_THINGS(N, _PP_REFLECT_ENUMERATOR, NS, T, __VA_ARGS__)))					
 #define _PP_REFLECT_STRUCT(N, NS, T, S, ...)			PP_REFLECT_STRUCT_ONLY(N, NS, T, S)				_PP_EXPAND(_PP_REFLECT_THINGS(N, PP_REFLECT_MEMBER, NS, T, __VA_ARGS__))
 #define _PP_REFLECT_STRUCT_TEMPLATE(N, NS, T, S, ...)	PP_REFLECT_STRUCT_TEMPLATE_ONLY(N, NS, T, S)	_PP_EXPAND(_PP_REFLECT_THINGS(N, PP_REFLECT_TEMPLATE_MEMBER, NS, T, __VA_ARGS__))
@@ -77,4 +77,4 @@ static constexpr void ForEachVar(Fn&& Visitor)
 template<class Ctti>
 concept Templated = requires { typename Ctti::TemplateArgs; };
 
-} // namespace PlainProps
\ No newline at end of file
+} // namespace PlainProps
diff --git a/Engine/Plugins/Experimental/PlainProps/Source/Tests/PlainPropsUeCoreTests.cpp b/Engine/Plugins/Experimental/PlainProps/Source/Tests/PlainPropsUeCoreTests.cpp
--- a/Engine/Plugins/Experimental/PlainProps/Source/Tests/PlainPropsUeCoreTests.cpp
+++ b/Engine/Plugins/Experimental/PlainProps/Source/Tests/PlainPropsUeCoreTests.cpp
@@ -47,7 +47,7 @@ struct FRuntimeIds
 	static FNameId				IndexName(FAnsiStringView Name)		{ return GNames.MakeName(FName(Name)); }
 	static FMemberId			IndexMember(FAnsiStringView Name)	{ return GNames.NameMember(FName(Name)); }
 	static FConcreteTypenameId	IndexTypename(FAnsiStringView Name)	{ return GNames.NameType(FName(Name)); }
-	static FScopeId				IndexScope(FAnsiStringView Name)	{ return GNames.MakeScope(FName(Name)); }
+	static FFlatScopeId			IndexScope(FAnsiStringView Name)	{ return GNames.NameScope(FName(Name)); }
 	static FEnumId				IndexEnum(FType Type)				{ return GNames.IndexEnum(Type); }
 	static FStructId			IndexStruct(FType Type)				{ return GNames.IndexStruct(Type); }
 	static FIdIndexerBase&		GetIndexer()						{ return GNames; }
@@ -76,16 +76,8 @@ struct FDeltaRuntime : FDefaultRuntime
 
 ////////////////////////////////////////////////////////////////////////////////////////////////
 
-template<typename Enum, EEnumMode Mode, class Runtime = FDefaultRuntime>
-struct TScopedEnumDeclaration
-{
-	using Ids = typename Runtime::Ids;
-	using Ctti = CttiOf<Enum>;
-
-	FEnumId Id;
-	TScopedEnumDeclaration() : Id(DeclareNativeEnum<Ctti, Ids>(Runtime::GetEnums(), Mode)) {}
-	~TScopedEnumDeclaration() { Runtime::GetEnums().Drop(Id); }
-};
+template<typename Enum, EEnumMode Mode>
+using TScopedDefaultEnumDeclaration = TScopedEnumDeclaration<Enum, Mode, FDefaultRuntime>;
 
 template<class T>
 using TScopedDefaultStructBinding = TScopedStructBinding<T, FDefaultRuntime>;
@@ -995,10 +987,10 @@ TEST_CASE_NAMED(FPlainPropsUeCoreTest, "System::Core::Serialization::PlainProps:
 	
 	SECTION("Enum")
 	{
-		TScopedEnumDeclaration<EFlat1, EEnumMode::Flat> Flat1;
-		TScopedEnumDeclaration<EFlat2, EEnumMode::Flat> Flat2;
-		TScopedEnumDeclaration<EFlag1, EEnumMode::Flag> Flag1;
-		TScopedEnumDeclaration<EFlag2, EEnumMode::Flag> Flag2;
+		TScopedDefaultEnumDeclaration<EFlat1, EEnumMode::Flat> Flat1;
+		TScopedDefaultEnumDeclaration<EFlat2, EEnumMode::Flat> Flat2;
+		TScopedDefaultEnumDeclaration<EFlag1, EEnumMode::Flag> Flag1;
+		TScopedDefaultEnumDeclaration<EFlag2, EEnumMode::Flag> Flag2;
 		TScopedDefaultStructBinding<FEnums> Int;
 		Run([](FBatchSaver& Batch)
 			{
@@ -1039,7 +1031,7 @@ TEST_CASE_NAMED(FPlainPropsUeCoreTest, "System::Core::Serialization::PlainProps:
 
 	SECTION("Nesting")
 	{
-		TScopedEnumDeclaration<EFlat1, EEnumMode::Flat> Flat1;
+		TScopedDefaultEnumDeclaration<EFlat1, EEnumMode::Flat> Flat1;
 		TScopedDefaultStructBinding<FLeafArrays> LeafArrays;
 		TScopedDefaultStructBinding<FComplexArrays> ComplexArrays;
 		Run([](FBatchSaver& Batch)
diff --git a/Engine/Plugins/Experimental/PlainPropsEngine/PlainPropsEngine.uplugin b/Engine/Plugins/Experimental/PlainPropsEngine/PlainPropsEngine.uplugin
--- a/Engine/Plugins/Experimental/PlainPropsEngine/PlainPropsEngine.uplugin
+++ b/Engine/Plugins/Experimental/PlainPropsEngine/PlainPropsEngine.uplugin
@@ -15,7 +15,7 @@
 			"Name": "PlainPropsEngine",
 			"Type": "Runtime",
 			"LoadingPhase": "PreLoadingScreen",
-			"PlatformAllowList": [ "Win64" ],
+			"PlatformAllowList": [ "Win64" ]
 		}
 	],
 	"Plugins": [
@@ -26,6 +26,10 @@
 		{
 			"Name": "PlainPropsUObject",
 			"Enabled": true
+		},
+		{
+			"Name": "Niagara",
+			"Enabled": true
 		}
 	]
 }
diff --git a/Engine/Plugins/Experimental/PlainPropsEngine/Source/PlainPropsEngine.Build.cs b/Engine/Plugins/Experimental/PlainPropsEngine/Source/PlainPropsEngine.Build.cs
--- a/Engine/Plugins/Experimental/PlainPropsEngine/Source/PlainPropsEngine.Build.cs
+++ b/Engine/Plugins/Experimental/PlainPropsEngine/Source/PlainPropsEngine.Build.cs
@@ -16,6 +16,7 @@ public class PlainPropsEngine : ModuleRules
 				"CoreUObject",
 				"CorePreciseFP",
 				"Engine",
+				"Niagara",
 				"PlainProps",
 				"PlainPropsUObject"
 		});
diff --git a/Engine/Plugins/Experimental/PlainPropsEngine/Source/Private/PlainPropsAnimationBindings.cpp b/Engine/Plugins/Experimental/PlainPropsEngine/Source/Private/PlainPropsAnimationBindings.cpp
new file mode 100644
--- /dev/null
+++ b/Engine/Plugins/Experimental/PlainPropsEngine/Source/Private/PlainPropsAnimationBindings.cpp
@@ -0,0 +1,286 @@
+// Copyright Epic Games, Inc. All Rights Reserved.
+
+#include "PlainPropsEngineBindings.h"
+#include "PlainPropsUObjectRuntime.h"
+#include "PlainPropsRoundtripTest.h" // EBindMode
+#include "PlainPropsBind.h"
+#include "PlainPropsLoad.h"
+#include "PlainPropsSave.h"
+#include "PlainPropsInternalPrivateMemberPtr.h"
+#include "Algo/IndexOf.h"
+#include "Animation/AnimData/AttributeIdentifier.h"
+#include "Animation/AttributeCurve.h"
+#include "Animation/AttributeTypes.h"
+#include "Misc/CoreMiscDefines.h"
+
+// Temp hacks for non-intrusive prototype
+PP_DEFINE_PRIVATE_MEMBER_NAMESPACE_BEGIN(FAnimationAttributeIdentifier_Private)
+PP_DEFINE_PRIVATE_MEMBER_PTR(FAnimationAttributeIdentifier, Name,					FName);
+PP_DEFINE_PRIVATE_MEMBER_PTR(FAnimationAttributeIdentifier, BoneName,				FName);
+PP_DEFINE_PRIVATE_MEMBER_PTR(FAnimationAttributeIdentifier, BoneIndex,				int32);
+PP_DEFINE_PRIVATE_MEMBER_PTR(FAnimationAttributeIdentifier, ScriptStruct,			TObjectPtr<UScriptStruct>);
+PP_DEFINE_PRIVATE_MEMBER_PTR(FAnimationAttributeIdentifier, ScriptStructPath,		FSoftObjectPath);
+PP_DEFINE_PRIVATE_MEMBER_NAMESPACE_END()
+
+PP_DEFINE_PRIVATE_MEMBER_NAMESPACE_BEGIN(FAttributeCurve_Private)
+PP_DEFINE_PRIVATE_MEMBER_PTR(FAttributeKey,				Value,						FWrappedAttribute);
+PP_DEFINE_PRIVATE_MEMBER_PTR(FAttributeCurve,			Keys,						TArray<FAttributeKey>);	
+PP_DEFINE_PRIVATE_MEMBER_PTR(FAttributeCurve,			ScriptStructPath,			FSoftObjectPath);
+PP_DEFINE_PRIVATE_MEMBER_PTR(FAttributeCurve,			ScriptStruct,				TObjectPtr<UScriptStruct>);
+PP_DEFINE_PRIVATE_MEMBER_PTR(FAttributeCurve,			bShouldInterpolate,			bool);
+PP_DEFINE_PRIVATE_MEMBER_PTR(FAttributeCurve,			Operator,					const UE::Anim::IAttributeBlendOperator*);
+PP_DEFINE_PRIVATE_MEMBER_NAMESPACE_END()
+
+namespace PlainProps::UE::Animation
+{
+
+static FScopeId GetAnimationScope()
+{
+	return GUE.Scopes.Engine;
+}
+
+struct FAnimationAttributeIdentifierBinding : ICustomBinding
+{
+	using Type = FAnimationAttributeIdentifier;
+	using Ids = FRuntimeIds;
+
+	const FMemberId MemberIds[4];
+	
+	FAnimationAttributeIdentifierBinding(TPropertySpecifier<4>& Spec)
+	: MemberIds{Ids::IndexMember("Name"), Ids::IndexMember("BoneName"), Ids::IndexMember("BoneIndex"), Ids::IndexMember("ScriptStructPath")}
+	{
+		Spec.Members[0] = FDeclId(GUE.Structs.Name);
+		Spec.Members[1] = FDeclId(GUE.Structs.Name);
+		Spec.Members[2] = Specify<int32>();
+		Spec.Members[3] = FDeclId(GUE.Structs.SoftObjectPath);
+	}
+
+	void Save(FMemberBuilder& Dst, const FAnimationAttributeIdentifier& Src, const FAnimationAttributeIdentifier* Default, const FSaveContext& Ctx) const
+	{
+		using namespace FAnimationAttributeIdentifier_Private;
+
+		Dst.AddStruct(MemberIds[0], GUE.Structs.Name, SaveStruct(&(Src.*_Name), GUE.Structs.Name, Ctx));
+		Dst.AddStruct(MemberIds[1], GUE.Structs.Name, SaveStruct(&(Src.*_BoneName), GUE.Structs.Name, Ctx));
+		Dst.Add(MemberIds[2], Src.*_BoneIndex);
+		Dst.AddStruct(MemberIds[3], GUE.Structs.SoftObjectPath, SaveStruct(&(Src.*_ScriptStructPath), GUE.Structs.SoftObjectPath, Ctx));
+	}
+
+	void Load(FAnimationAttributeIdentifier& Dst, FStructLoadView Src, ECustomLoadMethod Method) const
+	{
+		using namespace FAnimationAttributeIdentifier_Private;
+
+		FMemberLoader It(Src);
+		LoadStruct(&(Dst.*_Name), It.GrabStruct());
+		LoadStruct(&(Dst.*_BoneName), It.GrabStruct());
+		Dst.*_BoneIndex = It.GrabLeaf().As<int32>();
+		LoadStruct(&(Dst.*_ScriptStructPath), It.GrabStruct());
+		Dst.*_ScriptStruct = Cast<UScriptStruct>((Dst.*_ScriptStructPath).ResolveObject());
+	}
+
+	inline static bool Diff(FAnimationAttributeIdentifier A, FAnimationAttributeIdentifier B, const FBindContext&)
+	{
+		return A != B;
+	}
+};
+
+struct FAttributeIds
+{
+	static constexpr std::initializer_list<const char*> AnsiNames =
+	{
+		"IntegerAnimationAttribute",
+		"FloatAnimationAttribute",
+		"VectorAnimationAttribute",
+		"QuaternionAnimationAttribute",
+		"TransformAnimationAttribute",
+		"StringAnimationAttribute",
+		"NonBlendableIntegerAnimationAttribute",
+		"NonBlendableFloatAnimationAttribute",
+		"NonBlendableVectorAnimationAttribute",
+		"NonBlendableQuaternionAnimationAttribute",
+		"NonBlendableTransformAnimationAttribute",
+	};
+
+	FName Names[AnsiNames.size()];
+	FBindId BindIds[AnsiNames.size()];
+
+	template<class Ids>
+	FAttributeIds(TInit<Ids>)
+	{
+		FName* NamePtr = &Names[0];
+		FBindId* IdPtr = &BindIds[0];
+		for (const char* Name : AnsiNames)
+		{
+			*NamePtr++ = FName(Name);
+			*IdPtr++ = FBindId(Ids::IndexStruct(GetAnimationScope(), Name));
+		}
+	}
+
+	inline const FBindId GetBindId(UScriptStruct* Struct) const
+	{
+		int32 I = Algo::IndexOf(Names, Struct->GetFName());
+		checkf(I != INDEX_NONE, TEXT("Unknown attribute identifier '%s', it has not been indexed."), *Struct->GetFullName());
+		return BindIds[I];
+	}
+};
+
+struct FAttributeCurveBinding : ICustomBinding
+{
+	using Type = FAttributeCurve;
+	using Ids = FRuntimeIds;
+
+	const FDualStructId AttributeKeyId;
+	const FMemberId MemberIds[3];
+	FAttributeIds AttributeIds;
+	
+	FAttributeCurveBinding(TPropertySpecifier<3>& Spec)
+	: AttributeKeyId(GetStructDualId<Ids, FAttributeKey>())
+	, MemberIds{Ids::IndexMember("Keys"), Ids::IndexMember("ScriptStructPath"), Ids::IndexMember("DynamicKeys")}
+	, AttributeIds(TInit<Ids>{})
+	{
+		Spec.Members[0] = DefaultRangeOf(FDeclId(AttributeKeyId));
+		Spec.Members[1] = FDeclId(GUE.Structs.SoftObjectPath);
+		Spec.Members[2] = DefaultRangeOf(SpecDynamicStruct);
+	}
+
+	static FTypedRange SaveKeysStatic(FBindId Id, TConstArrayView<FAttributeKey> Keys, const FSaveContext& Ctx)
+	{
+		FStructRangeSaver Out(Ctx.Scratch, Keys.Num());
+		for (const FAttributeKey& Key : Keys)
+		{
+			Out.AddItem(SaveStruct(&Key, Id, Ctx));
+		}
+		return Out.Finalize(MakeStructRangeSchema(DefaultRangeMax, Id));
+	}
+
+	static FTypedRange SaveKeysDynamic(FBindId Id, TConstArrayView<FAttributeKey> Keys, const FSaveContext& Ctx)
+	{
+		using namespace FAttributeCurve_Private;
+
+		FStructRangeSaver Out(Ctx.Scratch, Keys.Num());
+		for (const FAttributeKey& Key : Keys)
+		{
+			Out.AddItem(SaveStruct((Key.*_Value).GetPtr<void>(), Id, Ctx));
+		}
+		return Out.Finalize(MakeStructRangeSchema(DefaultRangeMax, Id));
+	}
+
+	void Save(FMemberBuilder& Dst, const FAttributeCurve& Src, const FAttributeCurve* Default, const FSaveContext& Ctx) const
+	{
+		using namespace FAttributeCurve_Private;
+
+		Dst.AddRange(MemberIds[0], SaveKeysStatic(AttributeKeyId, Src.*_Keys, Ctx));
+		Dst.AddStruct(MemberIds[1], GUE.Structs.SoftObjectPath, SaveStruct(&(Src.*_ScriptStructPath), GUE.Structs.SoftObjectPath, Ctx));
+		FBindId DynamicKeyId = AttributeIds.GetBindId(Src.*_ScriptStruct);
+		Dst.AddRange(MemberIds[2], SaveKeysDynamic(DynamicKeyId, Src.*_Keys, Ctx));
+	}
+
+	static void LoadKeysStatic(TArray<FAttributeKey>& Dst, FStructRangeLoadView Src)
+	{
+		Dst.SetNum(static_cast<int32>(Src.Num()));
+		FAttributeKey* DstIt = Dst.GetData();
+		for (FStructLoadView SrcIt : Src)
+		{
+			LoadStruct(DstIt++, SrcIt);
+		}
+	}
+
+	static void LoadKeysDynamic(UScriptStruct* ScriptStruct, TArray<FAttributeKey>& Dst, FStructRangeLoadView Src)
+	{
+		using namespace FAttributeCurve_Private;
+
+		check(Dst.Num() == Src.Num());
+		FAttributeKey* DstIt = Dst.GetData();
+		for (FStructLoadView SrcIt : Src)
+		{
+			FAttributeKey& Item = *DstIt++;
+			(Item.*_Value).Allocate(ScriptStruct);
+			LoadStruct((Item.*_Value).GetPtr<void>(), SrcIt);
+		}
+	}
+
+	void Load(FAttributeCurve& Dst, FStructLoadView Src, ECustomLoadMethod Method) const
+	{
+		using namespace FAttributeCurve_Private;
+
+		FMemberLoader It(Src);
+		LoadKeysStatic(Dst.*_Keys, It.GrabRange().AsStructs());
+		LoadStruct(&(Dst.*_ScriptStructPath), It.GrabStruct());
+		UScriptStruct* ScriptStruct = Cast<UScriptStruct>((Dst.*_ScriptStructPath).ResolveObject());
+		LoadKeysDynamic(ScriptStruct, Dst.*_Keys, It.GrabRange().AsStructs());
+		Dst.*_ScriptStruct = ScriptStruct;
+		Dst.*_Operator = ::UE::Anim::AttributeTypes::GetTypeOperator(ScriptStruct);
+		Dst.*_bShouldInterpolate = ::UE::Anim::AttributeTypes::CanInterpolateType(ScriptStruct);
+	}
+
+	inline static bool Diff(FAttributeCurve A, FAttributeCurve B, const FBindContext&)
+	{
+		using namespace FAttributeCurve_Private;
+
+		if (A.*_ScriptStructPath != B.*_ScriptStructPath)
+		{
+			return true;
+		}
+		check(A.*_ScriptStruct == B.*_ScriptStruct);
+		if (A.*_Operator != B.*_Operator)
+		{
+			return true;
+		}
+		if (A.*_bShouldInterpolate != B.*_bShouldInterpolate)
+		{
+			return true;
+		}
+		TArray<FAttributeKey>& KeysA = A.*_Keys;
+		TArray<FAttributeKey>& KeysB = B.*_Keys;
+		if (KeysA.Num() != KeysB.Num())
+		{
+			return true;
+		}
+		UScriptStruct* ScriptStruct = A.*_ScriptStruct;
+		for (int32 I = 0, Num = KeysA.Num(); I < Num; ++I)
+		{
+			// compare the reflected properties
+			if (!FAttributeKey::StaticStruct()->CompareScriptStruct(&KeysA[I], &KeysB[I], PPF_None))
+			{
+				return true;
+			}
+			// compare the non-reflected dynamic struct properties
+			if (!ScriptStruct->CompareScriptStruct((KeysA[I].*_Value).GetPtr<void>(), (KeysB[I].*_Value).GetPtr<void>(), PPF_None))
+			{
+				return true;
+			}
+		}
+		return false;
+	}
+};
+
+} // namespace PlainProps::UE::Animation
+
+////////////////////////////////////////////////////////////////////////////////////////////////
+
+namespace PlainProps::UE
+{
+void CustomBindAnimationTypes(EBindMode Mode)
+{
+	using namespace Animation;
+	// Todo: Ownership / memory leak
+	new TScopedStructBinding<FAnimationAttributeIdentifier, FDefaultRuntime, Animation::FAnimationAttributeIdentifierBinding>();
+	new TScopedStructBinding<FAttributeCurve, FDefaultRuntime, Animation::FAttributeCurveBinding>();
+}
+}
+
+////////////////////////////////////////////////////////////////////////////////////////////////
+
+namespace PlainProps
+{
+template<> struct TOccupancyOf<FAnimationAttributeIdentifier> : FRequireAll {};
+template<> struct TOccupancyOf<FAttributeKey> : FRequireAll {};
+template<> struct TOccupancyOf<FAttributeCurve> : FRequireAll {};
+
+// Enable usage of TScopedStructBinding default constructor as well as GetStructXXXXId<Ids,Struct> and GetEnumId<Ids,Enum> lookups in custom bindings.
+template<> struct TTypename<FAnimationAttributeIdentifier>		{ inline static constexpr std::string_view DeclName = "AnimationAttributeIdentifier"; }; 
+template<> struct TTypename<FAttributeKey>						{ inline static constexpr std::string_view DeclName = "AttributeKey"; }; 
+template<> struct TTypename<FAttributeCurve>					{ inline static constexpr std::string_view DeclName = "AttributeCurve"; }; 
+template<> inline FScopeId GetOuterScope<UE::FRuntimeIds, TTypename<FAnimationAttributeIdentifier>>()		{ return UE::Animation::GetAnimationScope(); }
+template<> inline FScopeId GetOuterScope<UE::FRuntimeIds, TTypename<FAttributeKey>>()						{ return UE::Animation::GetAnimationScope(); }
+template<> inline FScopeId GetOuterScope<UE::FRuntimeIds, TTypename<FAttributeCurve>>()						{ return UE::Animation::GetAnimationScope(); }
+}
diff --git a/Engine/Plugins/Experimental/PlainPropsEngine/Source/Private/PlainPropsCommandlets.cpp b/Engine/Plugins/Experimental/PlainPropsEngine/Source/Private/PlainPropsCommandlets.cpp
--- a/Engine/Plugins/Experimental/PlainPropsEngine/Source/Private/PlainPropsCommandlets.cpp
+++ b/Engine/Plugins/Experimental/PlainPropsEngine/Source/Private/PlainPropsCommandlets.cpp
@@ -8,6 +8,7 @@
 #include "Misc/AsciiSet.h"
 #include "Misc/CommandLine.h"
 #include "Misc/PackageName.h"
+#include "ProfilingDebugging/CpuProfilerTrace.h"
 
 UTestPlainPropsCommandlet::UTestPlainPropsCommandlet(const FObjectInitializer& Init)
 : Super(Init)
@@ -21,49 +22,70 @@ int32 UTestPlainPropsCommandlet::Main(const FString& Params)
 	TArray<UObject*> Objects;
 	if (int32 LoadIdx = Params.Find("-load="); LoadIdx != INDEX_NONE)
 	{
+		TRACE_CPUPROFILER_EVENT_SCOPE(LoadAssets);
 		// E.g.
-		// -run=TestPlainProps -load=/BRRoot/BRRoot.BRRoot,/Game/Maps/FrontEnd.FrontEnd
-		// -run=TestPlainProps -load=/Game/Maps/Highrise.Highrise -save
+		// -load=/BRRoot/BRRoot.BRRoot,/Game/Maps/FrontEnd
+		// -load=maps
+		// -load=all
 
 		const TCHAR* It = &Params[LoadIdx + 6];
-		FStringView Assets(It, FAsciiSet::FindFirstOrEnd(It, FAsciiSet(" ")) - It);
-		UE_LOGFMT(LogPlainPropsEngine, Display, "Loading {Assets}...", Assets);
+		FStringView AssetString(It, FAsciiSet::FindFirstOrEnd(It, FAsciiSet(" ")) - It);
+		UE_LOGFMT(LogPlainPropsEngine, Display, "Loading {AssetString}...", AssetString);
 
-		int32 CommaIndex;
-		while (Assets.FindChar(',', CommaIndex))
+		const bool bLoadMaps = AssetString == TEXT("maps");
+		if (bLoadMaps || AssetString == TEXT("all"))
 		{
-			FSoftObjectPath AssetPath(Assets.Left(CommaIndex)); 
-			UObject* Object = AssetPath.TryLoad();
-			check(Object);
-			Objects.Add(Object);
-			Assets.RightChopInline(CommaIndex + 1);
-		}
+			UE_LOGFMT(LogPlainPropsEngine, Display, "Loading asset registry...");
+			FAssetRegistryModule& AssetRegistryModule = FModuleManager::LoadModuleChecked<FAssetRegistryModule>(AssetRegistryConstants::ModuleName);
+			AssetRegistryModule.Get().SearchAllAssets(true);
 
-		FSoftObjectPath AssetPath(Assets);
-		UObject* Object = AssetPath.TryLoad();
-		check(Object);
-		Objects.Add(Object);
-	}
-	else if (Params.Find("-loadmaps") != INDEX_NONE)
-	{
-		// load all .umaps in asset registry
-		UE_LOGFMT(LogPlainPropsEngine, Display, "Loading asset registry...");
-		FAssetRegistryModule& AssetRegistryModule = FModuleManager::LoadModuleChecked<FAssetRegistryModule>(AssetRegistryConstants::ModuleName);
-		AssetRegistryModule.Get().SearchAllAssets(true);
+			TArray<FAssetData> Assets;
 
-		UE_LOGFMT(LogPlainPropsEngine, Display, "Gathering all maps...");
-		TArray<FAssetData> Maps;
-		AssetRegistryModule.Get().GetAssetsByClass(FTopLevelAssetPath("/Script/Engine", "World"), /* out */ Maps, true);
+			if (bLoadMaps)
+			{
+				// find all .umaps in asset registry
+				UE_LOGFMT(LogPlainPropsEngine, Display, "Gathering all maps...");
+				AssetRegistryModule.Get().GetAssetsByClass(FTopLevelAssetPath("/Script/Engine", "World"), /* out */ Assets, true);
+				UE_LOGFMT(LogPlainPropsEngine, Display, "Loading all {Assets} maps...", Assets.Num());
+			}
+			else
+			{
+				// find all .uassets and .umaps in asset registry
+				UE_LOGFMT(LogPlainPropsEngine, Display, "Gathering all assets...");
+				AssetRegistryModule.Get().GetAllAssets(Assets, true);
+				UE_LOGFMT(LogPlainPropsEngine, Display, "Loading all {Assets} assets...", Assets.Num());
+			}
 
-		UE_LOGFMT(LogPlainPropsEngine, Display, "Loading all {Maps} maps...", Maps.Num());
-		for (FAssetData& Map : Maps)
+			for (FAssetData& Map : Assets)
+			{
+				FSoftObjectPath AssetPath = Map.GetSoftObjectPath();
+				AssetPath.LoadAsync({});
+			}
+		}
+		else
 		{
-			FSoftObjectPath AssetPath = Map.GetSoftObjectPath();
-			AssetPath.LoadAsync({});
+			int32 CommaIndex;
+			while (AssetString.FindChar(',', CommaIndex))
+			{
+				FSoftObjectPath AssetPath(AssetString.Left(CommaIndex)); 
+				UObject* Object = AssetPath.TryLoad();
+				check(Object);
+				Objects.Add(Object);
+				AssetString.RightChopInline(CommaIndex + 1);
+			}
+
+			FSoftObjectPath AssetPath(AssetString);
+			UObject* Object = AssetPath.TryLoad();
+			check(Object);
+			Objects.Add(Object);
 		}
 	}
 
-	FlushAsyncLoading();
+	UE_LOGFMT(LogPlainPropsEngine, Display, "FlushAsyncLoading ...");
+	{
+		TRACE_CPUPROFILER_EVENT_SCOPE(LoadAssets);
+		FlushAsyncLoading();
+	}
 	UE_LOGFMT(LogPlainPropsEngine, Display, "Starting test...");
 
 	UE_LOGFMT(LogPlainPropsEngine, Display, "Binding types to PlainProps schemas...");
@@ -78,11 +100,27 @@ int32 UTestPlainPropsCommandlet::Main(const FString& Params)
 		Mode = EBindMode::Runtime;
 	}
 
-	SchemaBindAllTypes(Mode);
-	CustomBindEngineTypes(Mode);
+	{
+		TRACE_CPUPROFILER_EVENT_SCOPE(BindAllTypes);
+		SchemaBindAllTypes(Mode);
+		CustomBindEngineTypes(Mode);
+		CustomBindAnimationTypes(Mode);
+		CustomBindNiagaraTypes(Mode);
+	}
 
-	// Parse roundtrip options
+	// Default roundtrip options
 	ERoundtrip Options = ERoundtrip::PP | ERoundtrip::UPS | ERoundtrip::TPS;
+
+	// Parse package option and if found always load PP with either one of TPS or UPS as a base
+	bool bPackages = Params.Find("-package") != INDEX_NONE;
+	if (bPackages)
+	{
+		Options = (Params.Find("-ups") != INDEX_NONE) ?
+			ERoundtrip::PP | ERoundtrip::UPS :
+			ERoundtrip::PP | ERoundtrip::TPS;
+	}
+
+	// Parse roundtrip options
 	if (Params.Find("-pp") != INDEX_NONE)
 	{
 		Options = ERoundtrip::PP | ERoundtrip::TextMemory;
@@ -93,7 +131,7 @@ int32 UTestPlainPropsCommandlet::Main(const FString& Params)
 	}
 
 	// Roundtrip
-	if (Params.Find("-package") != INDEX_NONE)
+	if (bPackages)
 	{
 		RoundtripViaPackages(Objects, Options);
 	}
@@ -102,6 +140,7 @@ int32 UTestPlainPropsCommandlet::Main(const FString& Params)
 		RoundtripViaBatch(Objects, Options);
 	}
 	
+	TRACE_CPUPROFILER_EVENT_FLUSH();
 	FPlatformMisc::RequestExit(true, TEXT("TestPlainProps"));
 	UE_LOGFMT(LogPlainPropsEngine, Display, "Done!");
 	return 0;
diff --git a/Engine/Plugins/Experimental/PlainPropsEngine/Source/Private/PlainPropsEngineBindings.cpp b/Engine/Plugins/Experimental/PlainPropsEngine/Source/Private/PlainPropsEngineBindings.cpp
--- a/Engine/Plugins/Experimental/PlainPropsEngine/Source/Private/PlainPropsEngineBindings.cpp
+++ b/Engine/Plugins/Experimental/PlainPropsEngine/Source/Private/PlainPropsEngineBindings.cpp
@@ -6,6 +6,7 @@
 #include "PlainPropsBind.h"
 #include "PlainPropsLoad.h"
 #include "PlainPropsSave.h"
+#include "PlainPropsVisualize.h"
 #include "Algo/Compare.h"
 #include "Engine/Texture2D.h"
 #include "Logging/StructuredLog.h"
@@ -141,6 +142,8 @@ struct FTexture2DCook : ICustomBinding
 
 void CustomBindEngineTypes(EBindMode Mode)
 {
+	{ DbgVis::FLoadStructPlan _; DbgVis::FIdVisualizer{}.KeepDebugInfo(&_); }
+
 	if (Mode == EBindMode::Runtime)
 	{
 		// Todo: Ownership / memory leak
diff --git a/Engine/Plugins/Experimental/PlainPropsEngine/Source/Private/PlainPropsEngineBindings.h b/Engine/Plugins/Experimental/PlainPropsEngine/Source/Private/PlainPropsEngineBindings.h
--- a/Engine/Plugins/Experimental/PlainPropsEngine/Source/Private/PlainPropsEngineBindings.h
+++ b/Engine/Plugins/Experimental/PlainPropsEngine/Source/Private/PlainPropsEngineBindings.h
@@ -10,5 +10,7 @@ namespace PlainProps::UE
 	enum class EBindMode : uint8;
 
 	void CustomBindEngineTypes(EBindMode Mode);
+	void CustomBindAnimationTypes(EBindMode Mode);
+	void CustomBindNiagaraTypes(EBindMode Mode);
 
-} // namespace PlainProps::UE
\ No newline at end of file
+} // namespace PlainProps::UE
diff --git a/Engine/Plugins/Experimental/PlainPropsEngine/Source/Private/PlainPropsNiagaraBindings.cpp b/Engine/Plugins/Experimental/PlainPropsEngine/Source/Private/PlainPropsNiagaraBindings.cpp
new file mode 100644
--- /dev/null
+++ b/Engine/Plugins/Experimental/PlainPropsEngine/Source/Private/PlainPropsNiagaraBindings.cpp
@@ -0,0 +1,148 @@
+// Copyright Epic Games, Inc. All Rights Reserved.
+
+#include "PlainPropsEngineBindings.h"
+#include "PlainPropsUObjectRuntime.h"
+#include "PlainPropsRoundtripTest.h" // EBindMode
+#include "PlainPropsBind.h"
+#include "PlainPropsLoad.h"
+#include "PlainPropsSave.h"
+#include "PlainPropsVisualize.h"
+#include "PlainPropsInternalPrivateMemberPtr.h"
+#include "NiagaraTypes.h"
+
+// Temp hacks for non-intrusive prototype
+PP_DEFINE_PRIVATE_MEMBER_NAMESPACE_BEGIN(FNiagaraTypeDefinition_Private)
+PP_DEFINE_PRIVATE_MEMBER_PTR(FNiagaraTypeDefinition,		Flags,						uint8);	
+PP_DEFINE_PRIVATE_MEMBER_PTR(FNiagaraTypeDefinition,		RegisteredTypeDefIndex,		int32);	
+PP_DEFINE_PRIVATE_MEMBER_PTR(FNiagaraTypeDefinition,		GeneratePathNameHash,		void());	
+PP_DEFINE_PRIVATE_MEMBER_NAMESPACE_END()
+
+PP_REFLECT_ENUM(FNiagaraTypeDefinition, FUnderlyingType, UT_None, UT_Class, UT_Struct, UT_Enum);
+// The last two constants in FNiagaraTypeDefinition::FTypeFlags TF_SerializedAsLWC and TF_AllowLWC_DEPRECATED should be unused,
+// leaving them out will cause a runtime assert if this is not the case.
+PP_REFLECT_ENUM(FNiagaraTypeDefinition, FTypeFlags, TF_None, TF_Static);
+
+////////////////////////////////////////////////////////////////////////////////////////////////
+
+namespace PlainProps::UE::Niagara
+{
+
+static FScopeId GetNiagaraScope()
+{
+	static FScopeId Scope(UE::FRuntimeIds::IndexScope("/Script/Niagara"));
+	return Scope;
+}
+
+// Custom bind FNiagaraTypeDefinition to modify the property values on Save and Load,
+// could use a schema binding for actual serialization.
+struct FNiagaraTypeDefinitionBinding : ICustomBinding
+{
+	using Type = FNiagaraTypeDefinition;
+	using Ids = FRuntimeIds;
+	using FUnderlyingType = FNiagaraTypeDefinition::FUnderlyingType;
+	using FTypeFlags = FNiagaraTypeDefinition::FTypeFlags;
+
+	const FEnumId UnderlyingTypeId;
+	const FEnumId TypeFlagsId;
+	const FMemberId MemberIds[3];
+	
+	FNiagaraTypeDefinitionBinding(TPropertySpecifier<3>& Spec)
+	: UnderlyingTypeId(GetEnumId<Ids, FUnderlyingType>())
+	, TypeFlagsId(GetEnumId<Ids, FTypeFlags>())
+	, MemberIds{Ids::IndexMember("ClassStructOrEnum"), Ids::IndexMember("UnderlyingType"), Ids::IndexMember("Flags")}
+	{
+		Spec.Members[0] = FDeclId(GUE.Structs.ObjectPtr);
+		Spec.Members[1] = Specify<FUnderlyingType>(UnderlyingTypeId);
+		Spec.Members[2] = Specify<FTypeFlags>(TypeFlagsId);
+	}
+
+	void Save(FMemberBuilder& Dst, const FNiagaraTypeDefinition& Src, const FNiagaraTypeDefinition* Default, const FSaveContext& Ctx) const
+	{
+		using namespace FNiagaraTypeDefinition_Private;
+
+		FBuiltStruct* ClassStructOrEnum = SaveStruct(&Src.ClassStructOrEnum, GUE.Structs.ObjectPtr, Ctx);
+		FUnderlyingType UnderlyingType = static_cast<FUnderlyingType>(Src.UnderlyingType);
+		FTypeFlags Flags = static_cast<FTypeFlags>(Src.GetFlags());
+
+		Dst.AddStruct(MemberIds[0], GUE.Structs.ObjectPtr, ClassStructOrEnum);
+		Dst.AddEnum(MemberIds[1], UnderlyingTypeId, UnderlyingType);
+		Dst.AddEnum(MemberIds[2], TypeFlagsId, Flags);
+	}
+
+	void Load(FNiagaraTypeDefinition& Dst, FStructLoadView Src, ECustomLoadMethod Method) const
+	{
+		using namespace FNiagaraTypeDefinition_Private;
+
+		FMemberLoader It(Src);
+		Dst.*_RegisteredTypeDefIndex = INDEX_NONE; // note: the private InvalidateCachedIndex() is not exported
+		LoadStruct(&Dst.ClassStructOrEnum, It.GrabStruct());
+		Dst.UnderlyingType = It.GrabLeaf().AsUnderlyingValue<uint16>();
+		Dst.*_Flags = It.GrabLeaf().AsUnderlyingValue<uint8>();
+		(Dst.*_GeneratePathNameHash)();
+	}
+
+	inline static bool Diff(FNiagaraTypeDefinition A, FNiagaraTypeDefinition B, const FBindContext&) { return A != B; }
+};
+
+// Custom bind FNiagaraTypeDefinitionHandle to Save and Load as a FNiagaraTypeDefinition
+struct FNiagaraTypeDefinitionHandleBinding : ICustomBinding
+{
+	using Ids = FRuntimeIds;
+	using Type = FNiagaraTypeDefinitionHandle;
+	const FDualStructId NiagaraTypeDefinitionId;
+	const FMemberId MemberIds[1];
+	
+	FNiagaraTypeDefinitionHandleBinding(TPropertySpecifier<1>& Spec)
+	: NiagaraTypeDefinitionId(GetStructDualId<Ids, FNiagaraTypeDefinition>())
+	, MemberIds{Ids::IndexMember("TypeDef")}
+	{
+		Spec.Members[0] = FDeclId(NiagaraTypeDefinitionId);
+	}
+
+	void Save(FMemberBuilder& Dst, const FNiagaraTypeDefinitionHandle& Src, const FNiagaraTypeDefinitionHandle* Default, const FSaveContext& Ctx) const
+	{
+		FNiagaraTypeDefinition TypeDef = *Src;
+		FBuiltStruct* BuiltStruct = SaveStruct(&TypeDef, NiagaraTypeDefinitionId, Ctx);
+		Dst.AddStruct(MemberIds[0], NiagaraTypeDefinitionId, BuiltStruct);
+	}
+
+	void Load(FNiagaraTypeDefinitionHandle& Dst, FStructLoadView Src, ECustomLoadMethod Method) const
+	{
+		FNiagaraTypeDefinition TypeDef;
+		LoadSoleStruct(&TypeDef, Src);
+		Dst = FNiagaraTypeDefinitionHandle(TypeDef);
+	}
+
+	inline static bool Diff(FNiagaraTypeDefinitionHandle A, FNiagaraTypeDefinitionHandle B, const FBindContext&) { return A != B; }
+};
+
+} // namespace PlainProps::UE::Niagara
+
+////////////////////////////////////////////////////////////////////////////////////////////////
+
+namespace PlainProps::UE
+{
+void CustomBindNiagaraTypes(EBindMode Mode)
+{
+	// Todo: Ownership / memory leak
+	new TScopedEnumDeclaration<FNiagaraTypeDefinition::FUnderlyingType, EEnumMode::Flat, FDefaultRuntime>();
+	new TScopedEnumDeclaration<FNiagaraTypeDefinition::FTypeFlags, EEnumMode::Flag, FDefaultRuntime>();
+	new TScopedStructBinding<FNiagaraTypeDefinition, FDefaultRuntime, Niagara::FNiagaraTypeDefinitionBinding>();
+	new TScopedStructBinding<FNiagaraTypeDefinitionHandle, FDefaultRuntime, Niagara::FNiagaraTypeDefinitionHandleBinding>();
+}
+} 
+
+////////////////////////////////////////////////////////////////////////////////////////////////
+namespace PlainProps
+{
+template<> struct TOccupancyOf<FNiagaraTypeDefinition> : FRequireAll {};
+template<> struct TOccupancyOf<FNiagaraTypeDefinitionHandle> : FRequireAll {};
+
+// Enable usage of TScopedStructBinding default constructor as well as GetStructXXXXId<Ids,Struct> and GetEnumId<Ids,Enum> lookups in custom bindings.
+template<> struct TTypename<FNiagaraTypeDefinition>			{ inline static constexpr std::string_view DeclName = "NiagaraTypeDefinition"; }; 
+template<> struct TTypename<FNiagaraTypeDefinitionHandle>	{ inline static constexpr std::string_view DeclName = "NiagaraTypeDefinitionHandle"; }; 
+template<> inline FScopeId GetOuterScope<UE::FRuntimeIds, TTypename<FNiagaraTypeDefinition>>()					{ return UE::Niagara::GetNiagaraScope(); }
+template<> inline FScopeId GetOuterScope<UE::FRuntimeIds, TTypename<FNiagaraTypeDefinitionHandle>>()			{ return UE::Niagara::GetNiagaraScope(); }
+template<> inline FScopeId GetOuterScope<UE::FRuntimeIds, TTypename<FNiagaraTypeDefinition::FUnderlyingType>>()	{ return UE::Niagara::GetNiagaraScope(); }
+template<> inline FScopeId GetOuterScope<UE::FRuntimeIds, TTypename<FNiagaraTypeDefinition::FTypeFlags>>()		{ return UE::Niagara::GetNiagaraScope(); }
+}
diff --git a/Engine/Plugins/Experimental/PlainPropsUObject/Source/Private/PlainPropsRoundtripTest.cpp b/Engine/Plugins/Experimental/PlainPropsUObject/Source/Private/PlainPropsRoundtripTest.cpp
--- a/Engine/Plugins/Experimental/PlainPropsUObject/Source/Private/PlainPropsRoundtripTest.cpp
+++ b/Engine/Plugins/Experimental/PlainPropsUObject/Source/Private/PlainPropsRoundtripTest.cpp
@@ -2,6 +2,7 @@
 
 #include "PlainPropsRoundtripTest.h"
 #include "PlainPropsUObjectBindingsInternal.h"
+#include "PlainPropsUObjectDiffInternal.h"
 #include "PlainPropsBuildSchema.h"
 #include "PlainPropsLoadMember.h"
 #include "PlainPropsUObjectRuntime.h"
@@ -29,6 +30,7 @@
 #include "UObject/SavePackage.h"
 #include "UObject/TextProperty.h"
 #include "UObject/UObjectIterator.h"
+#include "ProfilingDebugging/CpuProfilerTrace.h"
 
 #define PP_WITH_ENGINE_HOOKS 0
 
@@ -284,7 +286,7 @@ public:
 	void Save(FBindId Id, const UObject* Object, const UObject* Arch)
 	{
 		FBuiltStruct* Built;
-		if (Arch)
+		if (Arch && ShouldBind(Arch->GetClass()))
 		{
 			const UStruct* Class = SkipEmptyBases(Object->GetClass());
 			const UStruct* ArchClass = SkipEmptyBases(Arch->GetClass());
@@ -697,7 +699,7 @@ struct FInstance
 			Arch->GetClass()->GetDefaultObject(/* create lazily */ true);
 		}
 
-		Id = GUE.Names.IndexBindId(IndexType(Class));
+		Id = GUE.Names.IndexBindId(IndexStruct(Class));
 	}
 };
 
@@ -720,7 +722,6 @@ static bool IncludeClass(UClass* Class)
 		"GameFeaturePluginStateMachine", // Cloning ensure
 		"WorldSettings", // QAGame
 		// CitySample - Enum value 8 is undeclared in /Script/Engine.ERichCurveTangentMode, illegal value detected in /Script/Engine.RichCurveKey::TangentMode
-		"AnimSequence", "AnimationSequencerDataModel", "MovieSceneControlRigParameterSection",
 		// Enum value 7 is undeclared in /Script/Text3D.EText3DFontStyleFlags, illegal value detected in /Script/Text3DEditor.Text3DEditorFont::FontStyleFlags
 		// FontStyleFlags	Monospace | Bold | Italic (7 '\a')	UnrealEditor-Text3DEditor.dll!EText3DFontStyleFlags
 		"Text3DEditorFontSubsystem",
@@ -1006,7 +1007,7 @@ static void DumpUObjectToJsonFile(const UObject* Object)
 
 ////////////////////////////////////////////////////////////////////////////////
 
-#if WITH_EDITOR
+#if PP_WITH_ENGINE_HOOKS
 
 struct FLinkerPropertyBatch
 {
@@ -1083,7 +1084,7 @@ struct FLinkerPropertyBatch
 		return {Id, GUE.Members.Id, Specify<int32>()};
 	}
 
-	void Save(FMemberBuilder& Out, FObjectHandle In, const FSaveContext&)
+	void Save(FMemberBuilder& Out, const FObjectHandle& In, const FSaveContext&)
 	{
 		// see FLinkerSave::MapObject
 		UObject* Object = FObjectPtr(In).Get();
@@ -1107,28 +1108,23 @@ struct FLinkerPropertyBatch
 			if (PackageIdx.IsImport())
 			{
 				UObject* Object = LinkerLoad->ImportMap[PackageIdx.ToImport()].XObject;
-				if (Object == nullptr)
-				{
 #if PP_WITH_ENGINE_HOOKS
-					// Object = LinkerLoad->ResolveResource(PackageIdx);
+				//Object = Object ? Object : LinkerLoad->ResolveResource(PackageIdx);
 #endif
-				}
 				Out = ::UE::CoreUObject::Private::MakeObjectHandle(Object);
 			}
 			else if (PackageIdx.IsExport())
 			{
 				UObject* Object = LinkerLoad->ExportMap[PackageIdx.ToExport()].Object;
-				if (Object == nullptr)
-				{
 #if PP_WITH_ENGINE_HOOKS
-					// Object = LinkerLoad->ResolveResource(PackageIdx);
+				Object = Object ? Object : LinkerLoad->ResolveResource(PackageIdx);
 #endif
-				}
 				Out = ::UE::CoreUObject::Private::MakeObjectHandle(Object);
 			}
 			else
 			{
-				check(IsObjectHandleNull(Out));
+				check(PackageIdx.IsNull());
+				Out = FObjectHandle();
 			}
 		}
 		else
@@ -1204,7 +1200,8 @@ private:
 struct FPlainPropsLinkerSandbox
 {
 	static constexpr uint32 LoadFlags = (LOAD_Async | LOAD_NoVerify | LOAD_SkipLoadImportedPackages);
-	static constexpr uint32 CloneNumber = 8848;
+	static constexpr uint32 BaseNumber = 8848;
+	static constexpr uint32 PPNumber = 8849;
 
 	struct FSavedPackage
 	{
@@ -1214,15 +1211,36 @@ struct FPlainPropsLinkerSandbox
 		TArray64<uint8> PlainPropsData;
 	};
 
-	struct FLinkerBatch
+	struct FLinkerLoadInstance
 	{
-		FLinkerBatch(const FSavedPackage& Saved)
-			: PackagePath(Saved.PackagePath)
-			, PackageName(Saved.PackageName.GetNumber() ? FName(*Saved.PackageName.ToString(), CloneNumber) : FName(Saved.PackageName, CloneNumber))
+		FLinkerLoadInstance(const FSavedPackage& Saved, uint32 Number)
+			: PackageName(Saved.PackageName.GetNumber() ? FName(*Saved.PackageName.ToString(), Number) : FName(Saved.PackageName, Number))
 			, Package(NewObject<UPackage>(/*Outer*/nullptr, PackageName))
-			, LinkerReader(static_cast<const uint8*>(Saved.PackageData.GetData()), Saved.PackageData.GetSize(), ELargeMemoryReaderFlags::Persistent, PackageName)
-			, LinkerLoad(FLinkerLoad::CreateLinker(FUObjectThreadContext::Get().GetSerializeContext(), Package, PackagePath, LoadFlags, &LinkerReader))
-			, Properties(*LinkerLoad)
+			, LinkerReader(new FLargeMemoryReader(static_cast<const uint8*>(Saved.PackageData.GetData()), Saved.PackageData.GetSize(), ELargeMemoryReaderFlags::Persistent, PackageName))
+			, LinkerLoad(FLinkerLoad::CreateLinker(FUObjectThreadContext::Get().GetSerializeContext(), Package, Saved.PackagePath, LoadFlags, LinkerReader))
+		{
+			check(LinkerLoad->GetLoader_Unsafe() == LinkerReader);
+		}
+
+		~FLinkerLoadInstance()
+		{
+			// This will first destroy and free the LinkerReader then the LinkerLoad
+			ResetLoaders(Package);
+		}
+
+		FName PackageName;
+		UPackage* Package;
+		FLargeMemoryReader* LinkerReader;
+		FLinkerLoad* LinkerLoad;
+	};
+
+	struct FMountedPackage
+	{
+		FMountedPackage(const FSavedPackage& Saved)
+			: PackagePath(Saved.PackagePath)
+			, Base(Saved, BaseNumber)
+			, PP(Saved, PPNumber)
+			, Properties(*PP.LinkerLoad)
 			, Customs(Properties, GUE.Customs)
 			, BatchView(MakeMemoryView(Saved.PlainPropsData))
 			, BatchLoader(Customs.Overlay, BatchView)
@@ -1230,10 +1248,8 @@ struct FPlainPropsLinkerSandbox
 		}
 
 		FPackagePath PackagePath;
-		FName PackageName;
-		UPackage* Package;
-		FLargeMemoryReader LinkerReader;
-		FLinkerLoad* LinkerLoad;
+		FLinkerLoadInstance Base;
+		FLinkerLoadInstance PP;
 		FLinkerPropertyBatch Properties;
 		TCustomPropertyBindings<FLinkerPropertyBatch> Customs;
 		FMemoryView BatchView;
@@ -1241,11 +1257,24 @@ struct FPlainPropsLinkerSandbox
 	};
 
 	static TMap<FName, FSavedPackage> SavedPackages;
-	static TMap<FName, TUniquePtr<FLinkerBatch>> LoadBatches;
+	static TMap<FName, TUniquePtr<FMountedPackage>> MountedPackages;
+
+	// Transforms /Game/Package_01 which is represented as "/Game/Package" and number 1 into
+	// being represented as "/Game/Package_01" and number 0, so that number 0 can be replace by the two magic numbers 8848 and 8849.
+	// This works for packages named:/Game/Package_01, /Game/Package_02, ..., /Game/Package_10
+	static FName FoldNameNumber(FName Name)
+	{
+		if (Name.GetNumber())
+		{
+			Name = FName(*Name.ToString(), 0);
+		}
+		return Name;
+	}
 
 	static void SavePlainProps(const FLinkerSave& LinkerSave)
 	{
-		FName PackageName = LinkerSave.LinkerRoot->GetFName();
+		TRACE_CPUPROFILER_EVENT_SCOPE(SavePlainProps);
+		FName PackageName = FoldNameNumber(LinkerSave.LinkerRoot->GetFName());
 
 		UE_LOGFMT(LogPlainPropsUObject, Log, " Saving PlainProps for {PackageName}...", PackageName);
 		FLinkerPropertyBatch Properties(LinkerSave);
@@ -1254,7 +1283,7 @@ struct FPlainPropsLinkerSandbox
 		for (const FObjectExport& Export : LinkerSave.ExportMap)
 		{
 			const UStruct* Struct = SkipEmptyBases(Export.Object->GetClass());
-			FBindId Id = GUE.Names.IndexBindId(PlainProps::UE::IndexType(Struct));
+			FBindId Id = GUE.Names.IndexBindId(PlainProps::UE::IndexStruct(Struct));
 			if (GUE.Schemas.FindStruct(Id))
 			{
 				BatchSaver.Save(Id, Export.Object, Export.Object->GetArchetype());
@@ -1274,14 +1303,14 @@ struct FPlainPropsLinkerSandbox
 
 	static void SavePackageToMemory(const FLinkerSave& LinkerSave, uint8* Data, int64 Size)
 	{
-		FName PackageName = LinkerSave.LinkerRoot->GetFName();
+		FName PackageName = FoldNameNumber(LinkerSave.LinkerRoot->GetFName());
 		FSavedPackage& SavedPackage = SavedPackages.FindChecked(PackageName);
 		SavedPackage.PackageData = FSharedBuffer::TakeOwnership(Data, Size, FMemory::Free);
 	}
 
 	static bool OpenPackageReader(const FPackagePath& PackagePath, FArchive*& Ar)
 	{
-		FName PackageName = PackagePath.GetPackageFName();
+		FName PackageName = FoldNameNumber(PackagePath.GetPackageFName());
 		if (FSavedPackage* SavedData = SavedPackages.Find(PackageName))
 		{
 			Ar = new FLargeMemoryReader(static_cast<const uint8*>(SavedData->PackageData.GetData()), SavedData->PackageData.GetSize(), ELargeMemoryReaderFlags::Persistent, PackageName);
@@ -1290,41 +1319,92 @@ struct FPlainPropsLinkerSandbox
 		return false;
 	}
 
-	static void MountPackageFromMemory(FName PackageName)
+	static void MountPackagesFromMemory(FName PackageName)
 	{
+		PackageName = FoldNameNumber(PackageName);
 		if (FSavedPackage* SavedData = SavedPackages.Find(PackageName))
 		{
-			FLinkerBatch* LinkerBatch = new FLinkerBatch(*SavedData);
-			LoadBatches.Emplace(PackageName, LinkerBatch);
+			FMountedPackage* LinkerBatch = new FMountedPackage(*SavedData);
+			MountedPackages.Emplace(PackageName, LinkerBatch);
 		}
 	}
 
-	static void LoadPackageFromMemory(FName PackageName)
+	static void Shutdown()
 	{
-		if (TUniquePtr<FLinkerBatch>* LinkerBatch = LoadBatches.Find(PackageName))
+		MountedPackages.Empty();
+		SavedPackages.Empty();
+	}
+
+	static void LoadBaselinePackageFromMemory(FName PackageName)
+	{
+		PackageName = FoldNameNumber(PackageName);
+		if (TUniquePtr<FMountedPackage>* LinkerBatchPtr = MountedPackages.Find(PackageName))
 		{
+			FMountedPackage& LinkerBatch = **LinkerBatchPtr;
+
+			UE_LOGFMT(LogPlainPropsUObject, Log, " Loading baseline for {Package}...", PackageName);
+			LoadPackage(LinkerBatch.Base.Package, LinkerBatch.PackagePath, LoadFlags);
+		}
+	}
+
+	static void LoadPlainPropsPackageFromMemory(FName PackageName)
+	{
+		PackageName = FoldNameNumber(PackageName);
+		if (TUniquePtr<FMountedPackage>* LinkerBatchPtr = MountedPackages.Find(PackageName))
+		{
+			FMountedPackage& LinkerBatch = **LinkerBatchPtr;
+
 			UE_LOGFMT(LogPlainPropsUObject, Log, " Loading PlainProps for {Package}...", PackageName);
-			LoadPackage((*LinkerBatch)->Package, (*LinkerBatch)->PackagePath, LoadFlags);
+			LoadPackage(LinkerBatch.PP.Package, LinkerBatch.PackagePath, LoadFlags);
 		}
 	}
 
+	static void RoundtripToTextFromMemory(FName PackageName)
+	{
+		PackageName = FoldNameNumber(PackageName);
+		if (FSavedPackage* SavedData = SavedPackages.Find(PackageName))
+		{
+			FBatchTextRoundtripper StableBatch(MakeMemoryView(SavedData->PlainPropsData), ESchemaFormat::StableNames);
+			StableBatch.RoundtripTextForPackage(PackageName.ToString());
+		}
+	}
+
+	static FDiffObjectContext DiffPackagesLoadedFromMemory(FName PackageName)
+	{
+		PackageName = FoldNameNumber(PackageName);
+		if (TUniquePtr<FMountedPackage>* LinkerBatchPtr = MountedPackages.Find(PackageName))
+		{
+			FMountedPackage& LinkerBatch = **LinkerBatchPtr;
+
+			FDiffObjectContext Ctx = { LinkerBatch.Base.Package, LinkerBatch.PP.Package };
+			DiffObjects(Ctx);
+			return MoveTemp(Ctx);
+		}
+		return {};
+	}
+
 	static bool LoadObject(FStructuredArchive::FSlot& Slot, uint8* Data)
 	{
-		UObject* Object = (UObject*)Data;
-		FName PackageName = Object->GetPackage()->GetFName();
-		if (PackageName.GetNumber() == CloneNumber)
+		FArchive& Ar = Slot.GetUnderlyingArchive();
+		if (!Ar.IsLoading() || !Ar.IsPersistent() || Ar.HasAnyPortFlags(PPF_Duplicate|PPF_DuplicateForPIE))
 		{
-			PackageName.SetNumber(0);
+			return false;
 		}
 
-		TUniquePtr<FLinkerBatch>* LinkerBatch = LoadBatches.Find(PackageName);
-		if (!LinkerBatch)
+		UObject* Object = (UObject*)Data;
+		FName PackageName = Object->GetPackage()->GetFName();
+		if (PackageName.GetNumber() != PPNumber)
 		{
 			return false;
 		}
+		PackageName.SetNumber(0);
 
-		FPersistentBatchLoader& BatchLoader = (*LinkerBatch)->BatchLoader;
-		FLinkerLoad* LinkerLoad = (*LinkerBatch)->LinkerLoad;
+		TUniquePtr<FMountedPackage>* LinkerBatchPtr = MountedPackages.Find(PackageName);
+		check(LinkerBatchPtr);
+		FMountedPackage& LinkerBatch = **LinkerBatchPtr;
+
+		FPersistentBatchLoader& BatchLoader = LinkerBatch.BatchLoader;
+		FLinkerLoad* LinkerLoad = LinkerBatch.PP.LinkerLoad;
 		check(LinkerLoad == Object->GetLinker());
 
 		FUObjectSerializeContext* SerializeContext = FUObjectThreadContext::Get().GetSerializeContext();
@@ -1334,19 +1414,34 @@ struct FPlainPropsLinkerSandbox
 		const FObjectExport& Export = LinkerLoad->ExportMap[ExportIndex];
 		check(Export.Object == Object);
 
-		BatchLoader.LoadObject(Object, ExportIndex);
+		{
+			TRACE_CPUPROFILER_EVENT_SCOPE(LoadPlainProps);
+			BatchLoader.LoadObject(Object, ExportIndex);
+		}
+		if (const UClass* Class = Cast<UClass>(Object))
+		{
+			FBindId Id = GUE.Names.IndexBindId(PlainProps::UE::IndexStruct(Class));
+			UE_LOGFMT(LogPlainPropsUObject, Display, " Binding loaded dynamic class {BindId}", GUE.Debug.Print(Id));
+			check(!GUE.Schemas.FindStruct(Id))
+			UE::BindStruct(Id, Class);
+		}
 
-		const int32 Pos = Slot.GetUnderlyingArchive().Tell();
+#if WITH_EDITOR
+		const int32 Pos = Ar.Tell();
 		check(Pos == SerializeContext->SerializedObjectScriptStartOffset);
 		const int32 ScriptSerializationSize = Export.ScriptSerializationEndOffset - Export.ScriptSerializationStartOffset;
 		check(ScriptSerializationSize >= 0);
-		Slot.GetUnderlyingArchive().Seek(Pos + ScriptSerializationSize);
+		Ar.Seek(Pos + ScriptSerializationSize);
+#else
+		// todo: SerializedObjectScriptStartOffset does not exist in client builds
+		unimplemented();
+#endif
 		return true;
 	}
 };
 
 TMap<FName, FPlainPropsLinkerSandbox::FSavedPackage> FPlainPropsLinkerSandbox::SavedPackages;
-TMap<FName, TUniquePtr<FPlainPropsLinkerSandbox::FLinkerBatch>> FPlainPropsLinkerSandbox::LoadBatches;
+TMap<FName, TUniquePtr<FPlainPropsLinkerSandbox::FMountedPackage>> FPlainPropsLinkerSandbox::MountedPackages;
 
 #endif
 
@@ -1359,6 +1454,8 @@ static bool IsNativeCDO(UObject* Object)
 
 int32 RoundtripViaBatch(TConstArrayView<UObject*> Objects, ERoundtrip Options)
 {
+	TRACE_CPUPROFILER_EVENT_SCOPE(RoundtripViaBatch);
+
 	TArray<FInstance> Instances;
 	if (Objects.Num())
 	{
@@ -1595,6 +1692,8 @@ int32 RoundtripViaBatch(TConstArrayView<UObject*> Objects, ERoundtrip Options)
 
 int32 RoundtripViaPackages(TConstArrayView<UObject*> Objects, ERoundtrip Options)
 {
+	TRACE_CPUPROFILER_EVENT_SCOPE(RoundtripViaPackages);
+
 	TArray<UPackage*> Packages;
 	TArray<UObject*> AllPackages;
 
@@ -1632,12 +1731,6 @@ int32 RoundtripViaPackages(TConstArrayView<UObject*> Objects, ERoundtrip Options
 				bExclude = true;
 				return false;
 			}
-			if (!IncludeClass(Class))
-			{
-				UE_LOGFMT(LogPlainPropsUObject, Log, " Excluding PlainProps for {Package}", Object->GetPackage()->GetName());
-				bExclude = true;
-				return false;
-			}
 			return true;
 		});
 		if (bExclude)
@@ -1654,42 +1747,104 @@ int32 RoundtripViaPackages(TConstArrayView<UObject*> Objects, ERoundtrip Options
 	SavePackageSettings.ExternalWriteFileFunc = &FPlainPropsLinkerSandbox::SavePackageToMemory;
 	GSerializePropertiesOverride = &FPlainPropsLinkerSandbox::LoadObject;
 	GLoadFromPackageFileOverride = &FPlainPropsLinkerSandbox::OpenPackageReader;
+	FPlainPropsLinkerSandbox::SavedPackages.Reserve(Packages.Num());
+	FPlainPropsLinkerSandbox::MountedPackages.Reserve(Packages.Num());
 #endif
 
 	UE_LOGFMT(LogPlainPropsUObject, Display, "Saving {SaveNum} packages...", Packages.Num());
-	for (UPackage* Package : Packages)
+
 	{
-		UObject* Asset = Package->FindAssetInPackage();
-		FSavePackageArgs SaveArgs;
-		SaveArgs.TopLevelFlags = RF_Standalone;
-		SaveArgs.SaveFlags = SAVE_Async; // Write progress to temp memory instead of a temp file
-		const FString& Extension = Package->ContainsMap() ? FPackageName::GetMapPackageExtension() : FPackageName::GetAssetPackageExtension();
-		FString TempFilename = FPaths::Combine(*FPaths::ProjectSavedDir(), TEXT("PlainProps"), Package->GetName() + Extension);
-		// FString ActualFilename = FPackageName::LongPackageNameToFilename(Asset->GetPackage()->GetName(), Extension);
+		TRACE_CPUPROFILER_EVENT_SCOPE(SavePackages);
+		for (UPackage* Package : Packages)
+		{
+			UObject* Asset = Package->FindAssetInPackage();
+			FSavePackageArgs SaveArgs;
+			SaveArgs.TopLevelFlags = RF_Standalone;
+			SaveArgs.SaveFlags = SAVE_Async; // Write progress to temp memory instead of a temp file
+											 // Always use either one of TPS (default) or UPS for saving the baseline package
+			if (EnumHasAnyFlags(Options, ERoundtrip::UPS))
+			{
+				SaveArgs.SaveFlags |= SAVE_Unversioned_Properties;
+			}
+			const FString& Extension = Package->ContainsMap() ? FPackageName::GetMapPackageExtension() : FPackageName::GetAssetPackageExtension();
+			FString TempFilename = FPaths::Combine(*FPaths::ProjectSavedDir(), TEXT("PlainProps"), Package->GetName() + Extension);
+			// FString ActualFilename = FPackageName::LongPackageNameToFilename(Asset->GetPackage()->GetName(), Extension);
 
-		UPackage::Save(Package, Asset, *TempFilename, SaveArgs);
+			UPackage::Save(Package, Asset, *TempFilename, SaveArgs);
+		}
 	}
 
-#if WITH_EDITOR
-	FPlainPropsLinkerSandbox::SavedPackages.Reserve(Packages.Num());
-	FPlainPropsLinkerSandbox::LoadBatches.Reserve(Packages.Num());
-
-	UE_LOGFMT(LogPlainPropsUObject, Display, "Mounting {SaveNum} packages...", Packages.Num());
-	for (UPackage* Package : Packages)
+#if PP_WITH_ENGINE_HOOKS
+	// Roundtrip to text, also dumping all packages as yaml to disk
+	if (EnumHasAnyFlags(Options, ERoundtrip::TextMemory | ERoundtrip::TextStable))
 	{
-		FPlainPropsLinkerSandbox::MountPackageFromMemory(Package->GetFName());
+		TRACE_CPUPROFILER_EVENT_SCOPE(RoundtripToText);
+		UE_LOGFMT(LogPlainPropsUObject, Display, "Roundtripping {SaveNum} packages to text...", Packages.Num());
+		for (UPackage* Package : Packages)
+		{
+			FPlainPropsLinkerSandbox::RoundtripToTextFromMemory(Package->GetFName());
+		}
 	}
 
-	UE_LOGFMT(LogPlainPropsUObject, Display, "Loading {SaveNum} packages...", Packages.Num());
-	for (UPackage* Package : Packages)
+	// Roundtrip plain props and diff against baseline (TPS or UPS)
+	if (EnumHasAnyFlags(Options, ERoundtrip::PP))
 	{
-		FPlainPropsLinkerSandbox::LoadPackageFromMemory(Package->GetFName());
+		{
+			UE_LOGFMT(LogPlainPropsUObject, Display, "Mounting and loading {SaveNum} packages twice...", Packages.Num());
+			for (UPackage* Package : Packages)
+			{
+				FPlainPropsLinkerSandbox::MountPackagesFromMemory(Package->GetFName());
+			}
+		}
+		{
+			TRACE_CPUPROFILER_EVENT_SCOPE(LoadPlainPropsPackages);
+			double Start = FPlatformTime::Seconds();
+			for (UPackage* Package : Packages)
+			{
+				FPlainPropsLinkerSandbox::LoadPlainPropsPackageFromMemory(Package->GetFName());
+			}
+			UE_LOG(LogPlainPropsUObject, Display, TEXT("Loaded %d plainprops packages in %.2f s"),
+				Packages.Num(), FPlatformTime::Seconds() - Start);
+		}
+		{
+			TRACE_CPUPROFILER_EVENT_SCOPE(LoadBaselinePackags);
+			double Start = FPlatformTime::Seconds();
+			for (UPackage* Package : Packages)
+			{
+				FPlainPropsLinkerSandbox::LoadBaselinePackageFromMemory(Package->GetFName());
+			}
+			UE_LOG(LogPlainPropsUObject, Display, TEXT("Loaded %d baseline packages in %.2f s"),
+				Packages.Num(), FPlatformTime::Seconds() - Start);
+		}
+		{
+			TRACE_CPUPROFILER_EVENT_SCOPE(DiffPackages);
+			UE_LOGFMT(LogPlainPropsUObject, Display, "Diffing {SaveNum} packages...\n", Packages.Num());
+			TUtf8StringBuilder<4096> Diffs;
+			uint32 NumDiffs = 0;
+			for (UPackage* Package : Packages)
+			{
+				FDiffObjectContext Ctx = FPlainPropsLinkerSandbox::DiffPackagesLoadedFromMemory(Package->GetFName());
+				if (Ctx.Diffs.Num())
+				{
+					PrintDiff(Diffs, Ctx);
+					Diffs << " in " << Package->GetFName() << '\n';
+					++NumDiffs;
+				}
+			}
+			UE_LOGFMT(LogPlainPropsUObject, Display, "Detected {NumDiffs} diffs in {SaveNum} packages\n{Diffs}",
+				NumDiffs, Packages.Num(), Diffs.ToString());
+		}
 	}
-#endif
 
-#if PP_WITH_ENGINE_HOOKS
+	TRACE_CPUPROFILER_EVENT_SCOPE(Shutdown);
+#if UE_BUILD_DEBUG
+	// Note: Shutdown is too slow to be enabled by default.
+	// It will reset and flush all the linkers which takes ~15s for ~3000 packages in Release config
+	FPlainPropsLinkerSandbox::Shutdown();
+#endif
 	SavePackageSettings.ExternalWriteExportsFunc = nullptr;
 	GSerializePropertiesOverride = nullptr;
+	GLoadFromPackageFileOverride = nullptr;
 #endif
 
 	return 0;
diff --git a/Engine/Plugins/Experimental/PlainPropsUObject/Source/Private/PlainPropsUObjectBindings.cpp b/Engine/Plugins/Experimental/PlainPropsUObject/Source/Private/PlainPropsUObjectBindings.cpp
--- a/Engine/Plugins/Experimental/PlainPropsUObject/Source/Private/PlainPropsUObjectBindings.cpp
+++ b/Engine/Plugins/Experimental/PlainPropsUObject/Source/Private/PlainPropsUObjectBindings.cpp
@@ -372,7 +372,7 @@ static bool HasAny(uint64 Flags)
 
 ////////////////////////////////////////////////////////////////////////////////////////////////
 
-FType IndexType(const UField* Field)
+static FType IndexType(const UField* Field)
 {
 	check(Field);
 	FTypenameId Name = GUE.Names.MakeTypename(Field->GetFName());
@@ -387,6 +387,11 @@ FType IndexType(const UField* Field)
 	return { GUE.Names.NestReversedScopes(ReversedOuters), Name };
 }
 
+FType IndexStruct(const UStruct* Struct)
+{
+	return IndexType(Struct);
+}
+
 static EMemberPresence GetOccupancy(const UStruct* Struct)
 {
 	if (Struct->HasAnyCastFlags(CASTCLASS_UScriptStruct))
@@ -2650,8 +2655,9 @@ class FMapBindings
 			checkf(!Members.Value.IndexParameterName(GUE.Names).IsLowered(), TEXT("Only range-bound keys or values should be type-erased / lowered"));
 
 			FParametricTypeId PairTypename = GUE.Names.Resolve(Members.Pair.InnermostId.Get().AsStruct()).Name.AsParametric();
-			TConstArrayView<FType> Params = GUE.Names.Resolve(PairTypename).GetParameters();
-			FType Type = { GUE.Scopes.CoreUObject, FTypenameId(GUE.Names.MakeParametricTypeId(GUE.Typenames.Map, Params)) };
+			TConstArrayView<FType> PairParams = GUE.Names.Resolve(PairTypename).GetParameters();
+			FType MapParams[2] = { PairParams[0], PairParams[1] };
+			FType Type = { GUE.Scopes.CoreUObject, FTypenameId(GUE.Names.MakeParametricTypeId(GUE.Typenames.Map, MapParams)) };
 			FDualStructId Dual{GUE.Names.IndexStruct(Type)};
 			BothPair = { Dual, Dual };
 		}
@@ -3479,7 +3485,7 @@ static const UStruct* GetSuperToBind(const UStruct* Struct)
 	return nullptr;
 }
 
-static void BindStruct(FBindId Id, const UStruct* Struct)
+void BindStruct(FBindId Id, const UStruct* Struct)
 {
 	if (GUE.Customs.FindStruct(Id))
 	{
diff --git a/Engine/Plugins/Experimental/PlainPropsUObject/Source/Private/PlainPropsUObjectBindingsInternal.h b/Engine/Plugins/Experimental/PlainPropsUObject/Source/Private/PlainPropsUObjectBindingsInternal.h
--- a/Engine/Plugins/Experimental/PlainPropsUObject/Source/Private/PlainPropsUObjectBindingsInternal.h
+++ b/Engine/Plugins/Experimental/PlainPropsUObject/Source/Private/PlainPropsUObjectBindingsInternal.h
@@ -7,15 +7,15 @@
 
 DECLARE_LOG_CATEGORY_EXTERN(LogPlainPropsUObject, Log, All);
 
-class UField;
 class UStruct;
 
 namespace PlainProps::UE
 {
 
-FType			IndexType(const UField* Field);
+FType			IndexStruct(const UStruct* Struct);
+void			BindStruct(FBindId Id, const UStruct* Struct); // to bind dynamic structs as they are being loaded
 bool			ShouldBind(const UStruct* Struct);
-const UStruct*	SkipEmptyBases(const UStruct* In);
+const UStruct*	SkipEmptyBases(const UStruct* Struct);
 
 } // namespace PlainProps::UE
 
diff --git a/Engine/Plugins/Experimental/PlainPropsUObject/Source/Private/PlainPropsUObjectDiff.cpp b/Engine/Plugins/Experimental/PlainPropsUObject/Source/Private/PlainPropsUObjectDiff.cpp
new file mode 100644
--- /dev/null
+++ b/Engine/Plugins/Experimental/PlainPropsUObject/Source/Private/PlainPropsUObjectDiff.cpp
@@ -0,0 +1,479 @@
+// Copyright Epic Games, Inc. All Rights Reserved.
+
+#include "PlainPropsUObjectDiffInternal.h"
+#include "UObject/Object.h"
+#include "UObject/Package.h"
+#include "Containers/StringFwd.h"
+
+namespace PlainProps::UE
+{
+
+template<uint64 Mask>
+static bool HasAnyCastClassFlags(uint64 Flags)
+{
+	return (Mask & Flags) != 0;
+}
+
+// Inspired by UObjectBaseUtility::GetPathName and UObjectBaseUtility::GetFullGroupName.
+// Adjusted to handle the Object->HasAnyFlags(RF_HasExternalPackage) scenario.
+void GetRelativePathName(const UObject* Object, FStringBuilderBase& Out)
+{
+	const UObject* Outer = Object->GetOuter();
+	if (Outer)
+	{
+		if (const UObject* NextOuter = Outer->GetOuter())
+		{
+			GetRelativePathName(Outer, Out);
+
+			// SUBOBJECT_DELIMITER_CHAR is used to indicate that this object's outer is not a UPackage
+			if (Outer->GetClass() != UPackage::StaticClass() && NextOuter->GetClass() == UPackage::StaticClass())
+			{
+				Out << SUBOBJECT_DELIMITER_CHAR;
+			}
+			else
+			{
+				Out << TEXT('.');
+			}
+		}
+		Object->GetFName().AppendString(Out);
+	}
+}
+
+FString GetRelativePathName(const UObject* Object)
+{
+	TStringBuilder<256> Out;
+	GetRelativePathName(Object, Out);
+	return FString(FStringView(Out));
+}
+
+FAnsiStringView ToString(EDiffObjectNodeType Type)
+{
+	FAnsiStringView Types[5] = { ("NotInA"), ("NotInB"), ("Class"), ("Object"), ("Property") };
+	return Types[(uint8)Type];
+}
+
+///////////////////////////////////////////////////////////////////////////////
+
+bool DiffStructs(   const void* A, const void* B, const UStruct*   Struct,   FDiffObjectContext& Ctx);
+bool DiffProperties(const void* A, const void* B, const FProperty* Property, FDiffObjectContext& Ctx);
+
+bool DiffObjectReferences(const UObject* A, const UObject* B, FDiffObjectContext& Ctx)
+{
+	if (A == B)
+	{
+		return false;
+	}
+
+	const UClass* Class = A ? A->GetClass() : B->GetClass();
+	if (!A || !B)
+	{
+		Ctx.Diffs.Emplace(EDiffObjectNodeType::Object, A, B, Class);
+		return true;
+	}
+
+	if (Class != B->GetClass())
+	{
+		// Treat dynamic (non-native) classes as equal if their properties are equal.
+		if (Class->HasAnyClassFlags(CLASS_Native) ||
+			B->GetClass()->HasAnyClassFlags(CLASS_Native) ||
+			DiffObjectReferences(Class, B->GetClass(), Ctx))
+		{
+			Ctx.Diffs.Emplace(EDiffObjectNodeType::Class, A, B);
+			return true;
+		}
+	}
+
+	bool bIsInA = A->IsIn(Ctx.RootA);
+	bool bIsInB = B->IsIn(Ctx.RootB);
+	if (bIsInA & bIsInB)
+	{
+		TStringBuilder<256> PathA;
+		TStringBuilder<256> PathB;
+		GetRelativePathName(A, PathA);
+		GetRelativePathName(B, PathB);
+		bool bDiff = FStringView(PathA) != FStringView(PathB);
+		if (bDiff)
+		{
+			Ctx.Diffs.Emplace(EDiffObjectNodeType::Object, A, B, Class);
+		}
+		return bDiff;
+	}
+	else if (bIsInA | bIsInB)
+	{
+		Ctx.Diffs.Emplace(EDiffObjectNodeType::Object, A, B, Class);
+		return true;
+	}
+	return DiffStructs(A, B, Class, Ctx);
+}
+
+bool DiffObjects(const UObject* A, const UObject* B, FDiffObjectContext& Ctx)
+{
+	if (A == B)
+	{
+		return true;
+	}
+
+	const UClass* Class = A ? A->GetClass() : B->GetClass();
+	if (!A || !B)
+	{
+		Ctx.Diffs.Emplace(EDiffObjectNodeType::Object, A, B, Class);
+		return false;
+	}
+
+	if (Class != B->GetClass())
+	{
+		// Treat dynamic (non-native) classes as equal if their properties are equal.
+		if (Class->HasAnyClassFlags(CLASS_Native) ||
+			B->GetClass()->HasAnyClassFlags(CLASS_Native) ||
+			DiffObjectReferences(Class, B->GetClass(), Ctx))
+		{
+			Ctx.Diffs.Emplace(EDiffObjectNodeType::Class, A, B);
+			return true;
+		}
+	}
+
+	if (DiffStructs(A, B, Class, Ctx))
+	{
+		Ctx.Diffs.Emplace(EDiffObjectNodeType::Object, A, B, Class);
+		return true;
+	}
+
+	return false;
+}
+
+bool DiffStructProperties(const void* A, const void* B, const FStructProperty* Property, FDiffObjectContext& Ctx)
+{
+	const UScriptStruct* ScriptStruct = Property->Struct;
+	if (ScriptStruct->StructFlags & STRUCT_IdenticalNative)
+	{
+		bool bOut;
+		if (!ScriptStruct->GetCppStructOps()->Identical(A, B, /*PortFlags*/0, bOut))
+		{
+			Ctx.Diffs.Emplace(EDiffObjectNodeType::Property, A, B, Property);
+			return true;
+		}
+	}
+	if (DiffStructs(A, B, ScriptStruct, Ctx))
+	{
+		Ctx.Diffs.Emplace(EDiffObjectNodeType::Property, A, B, Property);
+		return true;
+	}
+	return false;
+}
+
+bool DiffObjectProperties(const void* A, const void* B, const FObjectPropertyBase* Property, FDiffObjectContext& Ctx)
+{
+	const UObject* ValueA = Property->GetObjectPropertyValue(A);
+	const UObject* ValueB = Property->GetObjectPropertyValue(B);
+	if (DiffObjectReferences(ValueA, ValueB, Ctx))
+	{
+		Ctx.Diffs.Emplace(EDiffObjectNodeType::Property, A, B, Property);
+		return true;
+	}
+	return false;
+}
+
+bool DiffArrayProperties(const void* A, const void* B, const FArrayProperty* Property, FDiffObjectContext& Ctx)
+{
+	FScriptArrayHelper ArrayA(Property, A);
+	FScriptArrayHelper ArrayB(Property, B);
+
+	if (ArrayA.Num() != ArrayB.Num())
+	{
+		Ctx.Diffs.Emplace(EDiffObjectNodeType::Property, A, B, Property);
+		return true;
+	}
+
+	for (int32 I = 0; I < ArrayA.Num(); I++)
+	{
+		const void* ValueA = ArrayA.GetRawPtr(I);
+		const void* ValueB = ArrayB.GetRawPtr(I);
+		if (DiffProperties(ValueA, ValueB, Property->Inner, Ctx))
+		{
+			Ctx.Diffs.Emplace(EDiffObjectNodeType::Property, A, B, Property, I);
+			return true;
+		}
+	}
+
+	return false;
+}
+
+bool DiffSetProperties(const void* A, const void* B, const FSetProperty* Property, FDiffObjectContext& Ctx)
+{
+	FScriptSetHelper SetA(Property, A);
+	FScriptSetHelper SetB(Property, B);
+	if (SetA.Num() != SetB.Num())
+	{
+		Ctx.Diffs.Emplace(EDiffObjectNodeType::Property, A, B, Property);
+		return true;
+	}
+
+	FScriptSetHelper::FIterator ItA(SetA);
+	FScriptSetHelper::FIterator ItB(SetB);
+	uint64 Idx = 0;
+	for (; ItA && ItB; ++ItA, ++ItB, ++Idx)
+	{
+		const void* ValueA = SetA.GetElementPtr(ItA);
+		const void* ValueB = SetB.GetElementPtr(ItB);
+		if (DiffProperties(ValueA, ValueB, Property->ElementProp, Ctx))
+		{
+			Ctx.Diffs.Emplace(EDiffObjectNodeType::Property, A, B, Property, Idx);
+			return true;
+		}
+	}
+	return false;
+}
+
+bool DiffMapProperties(const void* A, const void* B, const FMapProperty* Property, FDiffObjectContext& Ctx)
+{
+	FScriptMapHelper MapA(Property, A);
+	FScriptMapHelper MapB(Property, B);
+	if (MapA.Num() != MapB.Num())
+	{
+		Ctx.Diffs.Emplace(EDiffObjectNodeType::Property, A, B, Property);
+		return true;
+	}
+	FScriptMapHelper::FIterator ItA(MapA);
+	FScriptMapHelper::FIterator ItB(MapB);
+	uint64 Idx = 0;
+	for (; ItA && ItB; ++ItA, ++ItB, ++Idx)
+	{
+		const void* KeyA = MapA.GetKeyPtr(ItA);
+		const void* KeyB = MapB.GetKeyPtr(ItB);
+		if (DiffProperties(KeyA, KeyB, Property->KeyProp, Ctx))
+		{
+			Ctx.Diffs.Emplace(EDiffObjectNodeType::Property, A, B, Property, Idx);
+			return true;
+		}
+		const void* ValA = MapA.GetValuePtr(ItA);
+		const void* ValB = MapB.GetValuePtr(ItB);
+		if (DiffProperties(ValA, ValB, Property->ValueProp, Ctx))
+		{
+			Ctx.Diffs.Emplace(EDiffObjectNodeType::Property, A, B, Property, Idx);
+			return true;
+		}
+	}
+	return false;
+}
+
+bool DiffInterfaceProperties(const void* A, const void* B, const FInterfaceProperty* Property, FDiffObjectContext& Ctx)
+{
+	FScriptInterface* InterfaceA = (FScriptInterface*)A;
+	FScriptInterface* InterfaceB = (FScriptInterface*)B;
+	UObject* ObjectA = InterfaceA->GetObject();
+	UObject* ObjectB = InterfaceB->GetObject();
+	if (DiffObjectReferences(ObjectA, ObjectB, Ctx))
+	{
+		Ctx.Diffs.Emplace(EDiffObjectNodeType::Property, A, B, Property);
+		return true;
+	}
+	return false;
+}
+
+bool DiffProperties(const void* A, const void* B, const FProperty* Property, FDiffObjectContext& Ctx)
+{
+	uint64 Flags = Property->GetCastFlags();
+
+	if (HasAnyCastClassFlags<CASTCLASS_FStructProperty>(Flags))
+	{
+		if (DiffStructProperties(A, B, static_cast<const FStructProperty*>(Property), Ctx))
+		{
+			return true;
+		}
+	}
+	else if (HasAnyCastClassFlags<CASTCLASS_FObjectPropertyBase>(Flags))
+	{
+		if (DiffObjectProperties(A, B, static_cast<const FObjectPropertyBase*>(Property), Ctx))
+		{
+			return true;
+		}
+	}
+	else if (HasAnyCastClassFlags<CASTCLASS_FArrayProperty>(Flags))
+	{
+		if (DiffArrayProperties(A, B, static_cast<const FArrayProperty*>(Property), Ctx))
+		{
+			return true;
+		}
+	}
+	else if (HasAnyCastClassFlags<CASTCLASS_FSetProperty>(Flags))
+	{
+		if (DiffSetProperties(A, B, static_cast<const FSetProperty*>(Property), Ctx))
+		{
+			return true;
+		}
+	}
+	else if (HasAnyCastClassFlags<CASTCLASS_FMapProperty>(Flags))
+	{
+		if (DiffMapProperties(A, B, static_cast<const FMapProperty*>(Property), Ctx))
+		{
+			return true;
+		}
+	}
+	else if (HasAnyCastClassFlags<CASTCLASS_FInterfaceProperty>(Flags))
+	{
+		if (DiffInterfaceProperties(A, B, static_cast<const FInterfaceProperty*>(Property), Ctx))
+		{
+			return true;
+		}
+	}
+	else if (HasAnyCastClassFlags<CASTCLASS_FTextProperty>(Flags))
+	{
+		return false;
+	}
+	else
+	{
+		if (!Property->Identical(A, B, PPF_DeepComparison))
+		{
+			Ctx.Diffs.Emplace(EDiffObjectNodeType::Property, A, B, Property);
+			return true;
+		}
+	}
+	return false;
+}
+
+bool DiffStructs(const void* A, const void* B, const UStruct* Struct, FDiffObjectContext& Ctx)
+{
+	for (FProperty* Property = Struct->PropertyLink; Property; Property = Property->PropertyLinkNext)
+	{
+		if (!Property->ShouldDuplicateValue()) //Should the property be compared at all?
+		{
+			continue;
+		}
+		const void* ValueA = Property->ContainerPtrToValuePtr<void>(A);
+		const void* ValueB = Property->ContainerPtrToValuePtr<void>(B);
+		if (DiffProperties(ValueA, ValueB, Property, Ctx))
+		{
+			return true;
+		}
+	}
+	return false;
+}
+
+bool DiffObjects(FDiffObjectContext& Ctx)
+{
+	TMap<FString, const UObject*> ObjectsA;
+	TMap<FString, const UObject*> ObjectsB;
+
+	auto FillObjectMap = [](const UObject* Root, TMap<FString, const UObject*>& OutObjects)
+	{
+		OutObjects.Reserve(256);
+		ForEachObjectWithOuter(Root, [Root,&OutObjects](UObject* Obj)
+		{
+			// const UObject* StopOuter = Obj->HasAnyFlags(RF_HasExternalPackage) ? Obj->GetOutermostObject()->GetOuter() : Root;
+			OutObjects.Emplace(GetRelativePathName(Obj), Obj);
+			return true;
+		});
+	};
+	FillObjectMap(Ctx.RootA, ObjectsA);
+	FillObjectMap(Ctx.RootB, ObjectsB);
+
+	for (const TPair<FString, const UObject*>& PairA : ObjectsA)
+	{
+		const UObject* A = PairA.Value;
+		const UObject** B = ObjectsB.Find(PairA.Key);
+		if (!B)
+		{
+			Ctx.Diffs.Emplace(EDiffObjectNodeType::NotInB, A, nullptr, A->GetClass());
+			return true;
+		}
+
+		if (DiffObjects(A, *B, Ctx))
+		{
+			return true;
+		}
+	}
+
+	for (const TPair<FString, const UObject*>& PairB : ObjectsB)
+	{
+		const UObject* B = PairB.Value;
+		const UObject** A = ObjectsB.Find(PairB.Key);
+		if (!B)
+		{
+			Ctx.Diffs.Emplace(EDiffObjectNodeType::NotInA, B, nullptr, B->GetClass());
+			return true;
+		}
+	}
+	return false;
+}
+
+///////////////////////////////////////////////////////////////////////////////
+
+void PrintDiff(FUtf8StringBuilderBase& Out, FDiffObjectContext& Ctx)
+{
+	const FDiffObjectNode& LeafNode = Ctx.Diffs[0];
+	const FDiffObjectNode& RootNode = Ctx.Diffs.Last();
+
+	// Append diff type
+	Out << ToString(LeafNode.Type) << ": ";
+
+	// Build ('ValueA' != 'ValueB') string from the first/innermost node
+	Out << "('";
+	if (LeafNode.Type == EDiffObjectNodeType::Property)
+	{
+		const FProperty* Property = LeafNode.GetProperty();
+		FString A;
+		FString B;
+		Property->ExportTextItem_Direct(A, LeafNode.A, nullptr, (UObject*)Ctx.RootA, PPF_None, nullptr);
+		Property->ExportTextItem_Direct(B, LeafNode.B, nullptr, (UObject*)Ctx.RootB, PPF_None, nullptr);
+		Out << A << "' != '" << B;
+	}
+	else
+	{
+		const UObject* A = static_cast<const UObject*>(LeafNode.A);
+		const UObject* B = static_cast<const UObject*>(LeafNode.B);
+		Out << (A ? A->GetFullName() : "<null>");
+		Out << "' != '";
+		Out << (B ? B->GetFullName() : "<null>");
+	}
+	Out << "') in ";
+
+	// Build path root elemeent from the last/outermost UObject
+	check(RootNode.Type != EDiffObjectNodeType::Property);
+	if (&LeafNode != &RootNode)
+	{
+		const UObject* A = static_cast<const UObject*>(RootNode.A);
+		const UObject* B = static_cast<const UObject*>(RootNode.B);
+		if (const UClass* Class = RootNode.GetClass())
+		{
+			Out << (A ? A->GetName() : B->GetName());
+			Out << ':';
+			Out << Class->GetPrefixCPP() << Class->GetName();
+		}
+	}
+
+	// Build path chain based on all property nodes
+	bool bPrintPropertyPath = true;
+	for (const FDiffObjectNode& Diff : ReverseIterate(Ctx.Diffs))
+	{
+		if (Diff.Type != EDiffObjectNodeType::Property)
+		{
+			// all intermediate diff nodes should be properties
+			check(&Diff == &LeafNode || &Diff == &RootNode);
+			continue;
+		}
+		// Append .CPPMemberName:CPPMemberType to path string
+		// skip for inner properties coming after a range property, they just repeat the same type info again
+		if (bPrintPropertyPath)
+		{
+			const FProperty* Property = Diff.GetProperty();
+			FString ExtendedTypeText;
+			FString TypeText = Property->GetCPPType(&ExtendedTypeText);
+			Out << '.';
+			Out << Property->GetFName();
+			Out << ':';
+			Out << TypeText << ExtendedTypeText;
+		}
+		// Append range index
+		if (Diff.Idx != ~0ull)
+		{
+			Out << "[" << Diff.Idx << "]";
+			bPrintPropertyPath = false;
+		}
+		else
+		{
+			bPrintPropertyPath = true;
+		}
+	}
+}
+
+} // namespace PlainProps::UE
diff --git a/Engine/Plugins/Experimental/PlainPropsUObject/Source/Private/PlainPropsUObjectDiffInternal.h b/Engine/Plugins/Experimental/PlainPropsUObject/Source/Private/PlainPropsUObjectDiffInternal.h
new file mode 100644
--- /dev/null
+++ b/Engine/Plugins/Experimental/PlainPropsUObject/Source/Private/PlainPropsUObjectDiffInternal.h
@@ -0,0 +1,62 @@
+// Copyright Epic Games, Inc. All Rights Reserved.
+
+#pragma once
+
+#include "Containers/Array.h"
+#include "Containers/StringFwd.h"
+
+class FProperty;
+class UClass;
+class UObject;
+
+namespace PlainProps::UE
+{
+
+enum class EDiffObjectNodeType : uint8
+{
+	NotInA,		// An object under RootB is missing in RootA
+	NotInB,		// An object under RootA is missing in RootB
+	Class,		// Two object pointers A and B are different if:
+				// - both are set but the UClass is different, i.e. A->GetClass() != B->GetClass()
+	Object,		// Two object pointers A and B are different if:
+				// - only one of them is null, i.e. (A==nullptr != B==nullptr)
+				// - only one of them is an internal reference, i.e. (A->IsIn(RootA) != B->IsIn(RootB))
+				// - both are internal references but the relative paths inside their roots are different
+				// - both are external references and they have member properties with different values.
+	Property	// Two properties A and B have different values
+};
+
+struct FDiffObjectNode
+{
+	EDiffObjectNodeType Type;
+	const void* A;
+	const void* B;
+	const void* TypePtr = nullptr;
+	uint64 Idx = ~0ull;
+
+	const FProperty* GetProperty() const
+	{
+		return Type == EDiffObjectNodeType::Property ? (FProperty*)TypePtr : nullptr;
+	}
+
+	const UClass* GetClass() const
+	{
+		return Type != EDiffObjectNodeType::Property ? (UClass*)TypePtr : nullptr;
+	}
+};
+
+struct FDiffObjectContext
+{
+	// in
+	const UObject* const RootA = nullptr;
+	const UObject* const RootB = nullptr;
+
+	// out
+	TArray<FDiffObjectNode> Diffs;
+};
+
+bool DiffObjects(FDiffObjectContext& Ctx);
+
+void PrintDiff(FUtf8StringBuilderBase& Out, FDiffObjectContext& Ctx);
+
+}
diff --git a/Engine/Plugins/Experimental/PlainPropsUObject/Source/Public/PlainPropsUObjectRuntime.h b/Engine/Plugins/Experimental/PlainPropsUObject/Source/Public/PlainPropsUObjectRuntime.h
--- a/Engine/Plugins/Experimental/PlainPropsUObject/Source/Public/PlainPropsUObjectRuntime.h
+++ b/Engine/Plugins/Experimental/PlainPropsUObject/Source/Public/PlainPropsUObjectRuntime.h
@@ -144,12 +144,16 @@ PLAINPROPSUOBJECT_API extern FGlobals GUE;
 
 struct FRuntimeIds
 {
-	static FNameId				IndexName(FAnsiStringView Name)			{ return GUE.Names.MakeName(FName(Name)); }
-	static FMemberId			IndexMember(FAnsiStringView Name)		{ return GUE.Names.NameMember(FName(Name)); }
-	static FConcreteTypenameId	IndexTypename(FAnsiStringView Name)		{ return GUE.Names.NameType(FName(Name)); }
-	static FFlatScopeId			IndexScope(FAnsiStringView Name)		{ return GUE.Names.NameScope(FName(Name)); }
-	static FEnumId				IndexEnum(FType Type)					{ return GUE.Names.IndexEnum(Type); }
-	static FStructId			IndexStruct(FType Type)					{ return GUE.Names.IndexStruct(Type); }
+	static FNameId				IndexName(FAnsiStringView Name)						{ return GUE.Names.MakeName(FName(Name)); }
+	static FMemberId			IndexMember(FAnsiStringView Name)					{ return GUE.Names.NameMember(FName(Name)); }
+	static FConcreteTypenameId	IndexTypename(FAnsiStringView Name)					{ return GUE.Names.NameType(FName(Name)); }
+	static FFlatScopeId			IndexScope(FAnsiStringView Name)					{ return GUE.Names.NameScope(FName(Name)); }
+	static FEnumId				IndexEnum(FScopeId Scope, FAnsiStringView Name)		{ return GUE.Names.IndexEnum({Scope, GUE.Names.MakeTypename(Name)}); }
+	static FStructId			IndexStruct(FScopeId Scope, FAnsiStringView Name)	{ return GUE.Names.IndexStruct({Scope, GUE.Names.MakeTypename(Name)}); }
+	static FEnumId				IndexEnum(FType Type)								{ return GUE.Names.IndexEnum(Type); }
+	static FStructId			IndexStruct(FType Type)								{ return GUE.Names.IndexStruct(Type); }
+
+
 	static FIdIndexerBase&		GetIndexer()							{ return GUE.Names; }
 };
 
@@ -268,13 +272,13 @@ template<> struct TCustomBind<FReferencePropertyValue> { using Type = UE::FRefer
 template<> struct TCustomBind<FVerseFunction> { using Type = UE::FVerseFunctionBinding; };
 
 // Temporary way to tie certain types to /Script/CoreUObject scope
-template<> inline FScopeId IndexNamespaceId<UE::FRuntimeIds, TTypename<FIntPoint>>()		{ return UE::GUE.Scopes.CoreUObject; }
-template<> inline FScopeId IndexNamespaceId<UE::FRuntimeIds, TTypename<FQuat>>()			{ return UE::GUE.Scopes.CoreUObject; }
-template<> inline FScopeId IndexNamespaceId<UE::FRuntimeIds, TTypename<FTransform>>()		{ return UE::GUE.Scopes.CoreUObject; }
-template<> inline FScopeId IndexNamespaceId<UE::FRuntimeIds, TTypename<FVector>>()			{ return UE::GUE.Scopes.CoreUObject; }
-template<> inline FScopeId IndexNamespaceId<UE::FRuntimeIds, TTypename<FVector4>>()			{ return UE::GUE.Scopes.CoreUObject; }
-template<> inline FScopeId IndexNamespaceId<UE::FRuntimeIds, TTypename<FGuid>>()			{ return UE::GUE.Scopes.CoreUObject; }
-template<> inline FScopeId IndexNamespaceId<UE::FRuntimeIds, TTypename<FColor>>()			{ return UE::GUE.Scopes.CoreUObject; }
-template<> inline FScopeId IndexNamespaceId<UE::FRuntimeIds, TTypename<FLinearColor>>()		{ return UE::GUE.Scopes.CoreUObject; }
+template<> inline FScopeId GetOuterScope<UE::FRuntimeIds, TTypename<FIntPoint>>()		{ return UE::GUE.Scopes.CoreUObject; }
+template<> inline FScopeId GetOuterScope<UE::FRuntimeIds, TTypename<FQuat>>()			{ return UE::GUE.Scopes.CoreUObject; }
+template<> inline FScopeId GetOuterScope<UE::FRuntimeIds, TTypename<FTransform>>()		{ return UE::GUE.Scopes.CoreUObject; }
+template<> inline FScopeId GetOuterScope<UE::FRuntimeIds, TTypename<FVector>>()			{ return UE::GUE.Scopes.CoreUObject; }
+template<> inline FScopeId GetOuterScope<UE::FRuntimeIds, TTypename<FVector4>>()		{ return UE::GUE.Scopes.CoreUObject; }
+template<> inline FScopeId GetOuterScope<UE::FRuntimeIds, TTypename<FGuid>>()			{ return UE::GUE.Scopes.CoreUObject; }
+template<> inline FScopeId GetOuterScope<UE::FRuntimeIds, TTypename<FColor>>()			{ return UE::GUE.Scopes.CoreUObject; }
+template<> inline FScopeId GetOuterScope<UE::FRuntimeIds, TTypename<FLinearColor>>()	{ return UE::GUE.Scopes.CoreUObject; }
 
 }
