PlainProps: Fix bugs in the UObject diff code used in the commandlet package roundtrip scenario.

#rnx

diff --git a/Engine/Plugins/Experimental/PlainPropsUObject/Source/Private/PlainPropsUObjectDiff.cpp b/Engine/Plugins/Experimental/PlainPropsUObject/Source/Private/PlainPropsUObjectDiff.cpp
--- a/Engine/Plugins/Experimental/PlainPropsUObject/Source/Private/PlainPropsUObjectDiff.cpp
+++ b/Engine/Plugins/Experimental/PlainPropsUObject/Source/Private/PlainPropsUObjectDiff.cpp
@@ -110,14 +110,14 @@ bool DiffObjects(const UObject* A, const UObject* B, FDiffObjectContext& Ctx)
 {
 	if (A == B)
 	{
-		return true;
+		return false;
 	}
 
 	const UClass* Class = A ? A->GetClass() : B->GetClass();
 	if (!A || !B)
 	{
 		Ctx.Diffs.Emplace(EDiffObjectNodeType::Object, A, B, Class);
-		return false;
+		return true;
 	}
 
 	if (Class != B->GetClass())
@@ -386,8 +386,8 @@ bool DiffObjects(FDiffObjectContext& Ctx)
 	for (const TPair<FString, const UObject*>& PairB : ObjectsB)
 	{
 		const UObject* B = PairB.Value;
-		const UObject** A = ObjectsB.Find(PairB.Key);
-		if (!B)
+		const UObject** A = ObjectsA.Find(PairB.Key);
+		if (!A)
 		{
 			Ctx.Diffs.Emplace(EDiffObjectNodeType::NotInA, B, nullptr, B->GetClass());
 			return true;
