PlainProps: Add missing includes

#rnx

diff --git a/Engine/Plugins/Experimental/PlainPropsUObject/Source/Private/PlainPropsUObjectDiff.cpp b/Engine/Plugins/Experimental/PlainPropsUObject/Source/Private/PlainPropsUObjectDiff.cpp
--- a/Engine/Plugins/Experimental/PlainPropsUObject/Source/Private/PlainPropsUObjectDiff.cpp
+++ b/Engine/Plugins/Experimental/PlainPropsUObject/Source/Private/PlainPropsUObjectDiff.cpp
@@ -1,8 +1,10 @@
 // Copyright Epic Games, Inc. All Rights Reserved.
 
 #include "PlainPropsUObjectDiffInternal.h"
+#include "UObject/Class.h"
 #include "UObject/Object.h"
 #include "UObject/Package.h"
+#include "UObject/UnrealType.h"
 #include "Containers/StringFwd.h"
 
 namespace PlainProps::UE
