PlainProps: Add blueprints to linker roundtrip tests.
Add custom binding for FSoftObjectPath for linker roundtripping since source data uses a linker table index, and the actual paths can be mutated by FixupSoftObjectPathForInstancedPackage when the package summary is loaded.
Keep schema binding for FSoftObjectPath when roundtripping via plain memory.
Add specialized diff of SoftObjectPath to treat internal package references as equal when the internal relative paths are equal (same logic as for object references).
Add a EBatchType {Plain, Linker} enum that can be on command line as "-batch=linker", which replaces the "-package" option.
Add a simple custom binding for FGameplayEffectVersionBinding (in GameplayAbilities) to document this special case of always returning false from Identical() in order to force disable delta serialization in the old system.
Add a FLinkerDiffIgnoreFilter to RoundtripViaLinkerBatch.
Add and document diff exceptions in a new IgnoreLinkerDiff function.
Fix infinite recursion caused by Class->ClassDefaultObject->Class->ClassDefaultObject... comparisons in DiffStructs.

#rnx

diff --git a/Engine/Plugins/Experimental/PlainPropsEngine/PlainPropsEngine.uplugin b/Engine/Plugins/Experimental/PlainPropsEngine/PlainPropsEngine.uplugin
--- a/Engine/Plugins/Experimental/PlainPropsEngine/PlainPropsEngine.uplugin
+++ b/Engine/Plugins/Experimental/PlainPropsEngine/PlainPropsEngine.uplugin
@@ -27,6 +27,10 @@
 			"Name": "PlainPropsUObject",
 			"Enabled": true
 		},
+		{
+			"Name": "GameplayAbilities",
+			"Enabled": true
+		},
 		{
 			"Name": "Niagara",
 			"Enabled": true
diff --git a/Engine/Plugins/Experimental/PlainPropsEngine/Source/PlainPropsEngine.Build.cs b/Engine/Plugins/Experimental/PlainPropsEngine/Source/PlainPropsEngine.Build.cs
--- a/Engine/Plugins/Experimental/PlainPropsEngine/Source/PlainPropsEngine.Build.cs
+++ b/Engine/Plugins/Experimental/PlainPropsEngine/Source/PlainPropsEngine.Build.cs
@@ -16,6 +16,7 @@ public class PlainPropsEngine : ModuleRules
 				"CoreUObject",
 				"CorePreciseFP",
 				"Engine",
+				"GameplayAbilities",
 				"Niagara",
 				"PlainProps",
 				"PlainPropsUObject"
diff --git a/Engine/Plugins/Experimental/PlainPropsEngine/Source/Private/PlainPropsCommandlets.cpp b/Engine/Plugins/Experimental/PlainPropsEngine/Source/Private/PlainPropsCommandlets.cpp
--- a/Engine/Plugins/Experimental/PlainPropsEngine/Source/Private/PlainPropsCommandlets.cpp
+++ b/Engine/Plugins/Experimental/PlainPropsEngine/Source/Private/PlainPropsCommandlets.cpp
@@ -14,6 +14,125 @@ UTestPlainPropsCommandlet::UTestPlainPropsCommandlet(const FObjectInitializer& I
 : Super(Init)
 {}
 
+// Slow FName based IsChildOf implementation that avoids additional dependencies
+static bool IsChildOf(const UStruct* Struct, FName Base)
+{
+	for (const UStruct* It : Struct->GetSuperStructIterator())
+	{
+		if (It->GetFName() == Base)
+		{
+			return true;
+		}
+	}
+	return false;
+}
+
+// Temporary maintained filter to ignore diffs reported by RoundtripViaLinkerBatch
+// that currently relies on IsIdentical property comparisons of loaded objects.
+// This will be replaced with a more robust and scalable approach that also captures custom bound data.
+static bool IgnoreLinkerDiff(const void* A, const void* B, const FProperty* Property)
+{
+	const static FName GameplayEffectVersion{"GameplayEffectVersion"};
+	const static FName GraphGuid{"GraphGuid"};
+	const static FName LandscapeLayerInfoObject{"LandscapeLayerInfoObject"};
+	const static FName LayerUsageDebugColor{"LayerUsageDebugColor"};
+	const static FName NiagaraParameterDefinitionsBase{"NiagaraParameterDefinitionsBase"};
+	const static FName NodeGuid{"NodeGuid"};
+	const static FName OptimusDeformer{"OptimusDeformer"};
+	const static FName Status{"Status"};
+	const static FName UniqueId{"UniqueId"};
+
+
+	const static TSet<FName> MissingStructBindings
+	{
+		"InstancedStruct",			// FInstancedStruct in CoreUObject
+		"MovieSceneFrameRange",		// FMovieSceneFrameRange in MovieScene,
+									// requires Math bindings for TRange<T> and TRangeBound<T>
+		"ShaderValueTypeHandle",	// FShaderValueTypeHandle in ComputeFramework
+		"Spline",					// FSpline in Engine,
+									// requires GeometryCore binding for TMultiSpline
+	};
+
+	const uint64 CastFlags = Property->GetCastFlags();
+	const FName PropertyName = Property->GetFName();
+	const UStruct* OwnerStruct = Property->GetOwnerStruct();
+
+	// 1. Silence diffs for known missing custom bindings for custom serialized structs.
+	// These exceptions should be removed one by one as we implement these custom bindings.
+
+	if (CastFlags & CASTCLASS_FTextProperty) // FText is tricky to bind non-intrusively
+	{
+		return true;
+	}
+	else if (const FStructProperty* StructProp = CastField<FStructProperty>(Property))
+	{
+		// check missing bindings for struct properties
+		FName StructName = StructProp->Struct->GetFName();
+		if (MissingStructBindings.Contains(StructName))
+		{
+			return true;
+		}
+		else if (StructName == GameplayEffectVersion)
+		{
+			// FGameplayEffectVersion::Identical always return false to disable UE delta serialization,
+			// don't report this as a diff when the data is actually identical
+			return FMemory::Memcmp(A, B, StructProp->Struct->GetCppStructOps()->GetSize()) == 0;
+		}
+	}
+	else if (const FMapProperty* MapProp = CastField<FMapProperty>(Property))
+	{
+		// check missing bindings for struct properties inside map properties
+		if (const FStructProperty* KeyProp = CastField<FStructProperty>(MapProp->GetKeyProperty()))
+		{
+			if (MissingStructBindings.Contains(KeyProp->Struct->GetFName()))
+			{
+				return true;
+			}
+		}
+		if (const FStructProperty* ValueProp = CastField<FStructProperty>(MapProp->GetValueProperty()))
+		{
+			if (MissingStructBindings.Contains(ValueProp->Struct->GetFName()))
+			{
+				return true;
+			}
+		}
+	}
+	else if (OwnerStruct->GetFName() == OptimusDeformer)
+	{
+		// Symptom from the missing FText binding.
+		// UOptimusNode::PostLoad() calls UOptimusNode_CustomComputeKernel::PostLoadNodeSpecificData()
+		// which is doing SetDisplayName()/UOptimusNodeGraph::Notify()/Deformer->MarkModified()
+		// that sets status based on a modified FText
+		return PropertyName == Status;
+	}
+
+	// 2. Ignore diffs for runtime mutated properties, e.g. properties were
+	// generation is non-deterministic, session uniqueness is enforced, generation is dependent on path name, etc
+
+	if (PropertyName == GraphGuid)
+	{
+		// UEdGraph::PostInitProperties() calls FGuid::NewGuid();
+		return IsChildOf(OwnerStruct, "EdGraph");
+	}
+	else if (PropertyName == NodeGuid)
+	{
+		// UEdGraphNode::CreateDeterministicGuid() calls FGuid::NewDeterministicGuid(GetPathName());
+		return IsChildOf(OwnerStruct, "EdGraphNode");
+	}
+	else if (PropertyName == LayerUsageDebugColor)
+	{
+		// ULandscapeLayerInfoObject constructor is calls GenerateLayerUsageDebugColor() that depends on GetPathName()
+		return OwnerStruct->GetFName() == LandscapeLayerInfoObject;
+	}
+	else if (OwnerStruct->GetFName() == NiagaraParameterDefinitionsBase)
+	{
+		// UNiagaraParameterDefinitions::PostLoad() is doing FNiagaraEditorModule::EnsureReservedDefinitionUnique() to force a new unique id
+		return PropertyName == UniqueId;
+	}
+
+	return false;
+};
+
 int32 UTestPlainPropsCommandlet::Main(const FString& Params)
 {
 	using namespace PlainProps;
@@ -100,9 +219,15 @@ int32 UTestPlainPropsCommandlet::Main(const FString& Params)
 		Mode = EBindMode::Runtime;
 	}
 
+	EBatchType BatchType = EBatchType::Plain;
+	if (Params.Find("-batch=linker") != INDEX_NONE)
+	{
+		BatchType = EBatchType::Linker;
+	}
+
 	{
 		TRACE_CPUPROFILER_EVENT_SCOPE(BindAllTypes);
-		SchemaBindAllTypes(Mode);
+		SchemaBindAllTypes(Mode, BatchType);
 		CustomBindEngineTypes(Mode);
 		CustomBindAnimationTypes(Mode);
 		CustomBindNiagaraTypes(Mode);
@@ -111,9 +236,8 @@ int32 UTestPlainPropsCommandlet::Main(const FString& Params)
 	// Default roundtrip options
 	ERoundtrip Options = ERoundtrip::PP | ERoundtrip::UPS | ERoundtrip::TPS;
 
-	// Parse package option and if found always load PP with either one of TPS or UPS as a base
-	bool bPackages = Params.Find("-package") != INDEX_NONE;
-	if (bPackages)
+	// For linker batches always load PP with either one of TPS or UPS as a base
+	if (BatchType == EBatchType::Linker)
 	{
 		Options = (Params.Find("-ups") != INDEX_NONE) ?
 			ERoundtrip::PP | ERoundtrip::UPS :
@@ -131,13 +255,13 @@ int32 UTestPlainPropsCommandlet::Main(const FString& Params)
 	}
 
 	// Roundtrip
-	if (bPackages)
+	if (BatchType == EBatchType::Linker)
 	{
-		RoundtripViaPackages(Objects, Options);
+		RoundtripViaLinkerBatch(Objects, Options, IgnoreLinkerDiff);
 	}
 	else
 	{ 
-		RoundtripViaBatch(Objects, Options);
+		RoundtripViaPlainBatch(Objects, Options);
 	}
 	
 	TRACE_CPUPROFILER_EVENT_FLUSH();
diff --git a/Engine/Plugins/Experimental/PlainPropsEngine/Source/Private/PlainPropsEngineBindings.cpp b/Engine/Plugins/Experimental/PlainPropsEngine/Source/Private/PlainPropsEngineBindings.cpp
--- a/Engine/Plugins/Experimental/PlainPropsEngine/Source/Private/PlainPropsEngineBindings.cpp
+++ b/Engine/Plugins/Experimental/PlainPropsEngine/Source/Private/PlainPropsEngineBindings.cpp
@@ -9,6 +9,7 @@
 #include "PlainPropsVisualize.h"
 #include "Algo/Compare.h"
 #include "Engine/Texture2D.h"
+#include "GameplayEffect.h"
 #include "Logging/StructuredLog.h"
 #include "Misc/DefinePrivateMemberPtr.h"
 
@@ -140,15 +141,66 @@ struct FTexture2DCook : ICustomBinding
 
 ////////////////////////////////////////////////////////////////////////////////////////////////
 
+static FScopeId GetAbilitiesScope()
+{
+	static FScopeId Scope(UE::FRuntimeIds::IndexScope("/Script/GameplayAbilities"));
+	return Scope;
+}
+
+struct FGameplayEffectVersionBinding : ICustomBinding
+{
+	using Type = FGameplayEffectVersion;
+	using Ids = FRuntimeIds;
+
+	const FEnumId GameplayEffectVersionId;
+	const FMemberId MemberIds[1];
+	
+	FGameplayEffectVersionBinding(TPropertySpecifier<1>& Spec)
+	: GameplayEffectVersionId(Ids::IndexEnum(GetAbilitiesScope(), "EGameplayEffectVersion"))
+	, MemberIds{Ids::IndexMember("CurrentVersion")}
+	{
+		Spec.Members[0] = Specify<EGameplayEffectVersion>(GameplayEffectVersionId);
+	}
+
+	void Save(FMemberBuilder& Dst, const FGameplayEffectVersion& Src, const FGameplayEffectVersion* Default, const FSaveContext& Ctx) const
+	{
+		Dst.AddEnum(MemberIds[0], GameplayEffectVersionId, Src.CurrentVersion);
+	}
+
+	inline static void Plan(FMemcpyLoadPlan& Out)
+	{
+		check(sizeof(FGameplayEffectVersion) == sizeof(EGameplayEffectVersion));
+		Out.Size = sizeof(FGameplayEffectVersion);
+	}
+	
+	inline static bool Diff(FGameplayEffectVersion A, FGameplayEffectVersion B, const FBindContext&)
+	{
+		// TODO: do we need support for opting out of delta serialization?
+		// The existing custom serialize implementation always returns false from Identical as a hack to disable delta serialization.
+		check(A.Identical(&B, /*PortFlags*/0) == false); // detect if the existing code is changed
+		return A.CurrentVersion != B.CurrentVersion;
+	}
+};
+
+////////////////////////////////////////////////////////////////////////////////////////////////
+
 void CustomBindEngineTypes(EBindMode Mode)
 {
 	{ DbgVis::FLoadStructPlan _; DbgVis::FIdVisualizer{}.KeepDebugInfo(&_); }
 
+	// Todo: Ownership / memory leak
+	new TScopedStructBinding<FGameplayEffectVersion, FDefaultRuntime, FGameplayEffectVersionBinding>();
 	if (Mode == EBindMode::Runtime)
 	{
-		// Todo: Ownership / memory leak
 		new TScopedStructBinding<UTexture2D, FDeltaRuntime, FTexture2DCook>(FDualStructId(IndexEngineStruct("Texture2D")));
 	}
 }
 
 } // namespace PlainProps::UE
+
+namespace PlainProps
+{
+template<> struct TOccupancyOf<FGameplayEffectVersion> : FRequireAll {};
+template<> struct TTypename<FGameplayEffectVersion> { inline static constexpr std::string_view DeclName = "GameplayEffectVersion"; }; 
+template<> inline FScopeId GetOuterScope<UE::FRuntimeIds, TTypename<FGameplayEffectVersion>>() { return UE::GetAbilitiesScope(); }
+}
diff --git a/Engine/Plugins/Experimental/PlainPropsUObject/Source/Private/PlainPropsRoundtripTest.cpp b/Engine/Plugins/Experimental/PlainPropsUObject/Source/Private/PlainPropsRoundtripTest.cpp
--- a/Engine/Plugins/Experimental/PlainPropsUObject/Source/Private/PlainPropsRoundtripTest.cpp
+++ b/Engine/Plugins/Experimental/PlainPropsUObject/Source/Private/PlainPropsRoundtripTest.cpp
@@ -143,6 +143,7 @@ struct FMemoryPropertyBatch
 inline bool DiffProperty(FName A, FName B) { return !A.IsEqual(B, ENameCase::CaseSensitive); }
 inline bool DiffProperty(const FText& A, const FText& B) { return !FTextProperty::Identical_Implementation(A, B, 0); }
 inline bool DiffProperty(FObjectHandle A, FObjectHandle B) { return A != B; }
+inline bool DiffProperty(const FSoftObjectPath& A, const FSoftObjectPath& B) { return A != B; }
 inline bool DiffProperty(const FWeakObjectPtr& A, const FWeakObjectPtr& B) { return A != B; }
 inline bool DiffProperty(const FSoftObjectPtr& A, const FSoftObjectPtr& B) { return A != B; }
 inline bool DiffProperty(const FLazyObjectPtr& A, const FLazyObjectPtr& B) { return A != B; }
@@ -1140,7 +1141,7 @@ struct FLinkerPropertyBatch
 
 	void Save(FMemberBuilder& Out, const FWeakObjectPtr& In, const FSaveContext& Ctx)
 	{
-		// Save ObjectSerialNumber + ObjectIndex a single uint64
+		// Save ObjectSerialNumber + ObjectIndex as a single uint64
 		UObject* Object = In.Get(true);
 		Save(Out, ::UE::CoreUObject::Private::MakeObjectHandle(Object), Ctx);
 	}
@@ -1152,21 +1153,52 @@ struct FLinkerPropertyBatch
 		Out = ::UE::CoreUObject::Private::ResolveObjectHandle(Handle);
 	}
 
+	static FSoleMemberSpec Spec(FSoftObjectPath*, FDeclId Id)
+	{
+		return {Id, GUE.Members.Id, Specify<int32>()};
+	}
+
+	void Save(FMemberBuilder& Out, const FSoftObjectPath& In, const FSaveContext& Ctx)
+	{
+		// see FLinkerSave::MapSoftObjectPath
+		int32 Idx = LinkerSave->SoftObjectPathIndices.FindRef(In, INDEX_NONE);
+		if (Idx == INDEX_NONE)
+		{
+			UE_LOGFMT(LogPlainPropsUObject, Warning, "SoftObjectPath {Path} is not mapped.", In);
+		}
+		Out.Add(GUE.Members.Id, Idx);
+	}
+
+	void Load(FSoftObjectPath& Out, FStructLoadView In)
+	{
+		int32 Idx;
+		LoadSole<int32>(&Idx, In);
+		if (Idx != INDEX_NONE)
+		{
+			check(Idx < LinkerLoad->SoftObjectPathList.Num());
+			Out = LinkerLoad->SoftObjectPathList[Idx];
+		}
+		else
+		{
+			Out.Reset();
+		}
+	}
+
 	static FSoleMemberSpec Spec(FSoftObjectPtr*, FDeclId Id)
 	{
-		return {Id, GUE.Members.Id, FMemberSpec(GUE.Structs.SoftObjectPath)};
+		return Spec((FSoftObjectPath*)nullptr, Id);
 	}
 
-	static void Save(FMemberBuilder& Out, const FSoftObjectPtr& In, const FSaveContext& Ctx)
+	void Save(FMemberBuilder& Out, const FSoftObjectPtr& In, const FSaveContext& Ctx)
 	{
-		FBuiltStruct* SoftPath = SaveStruct(&In.GetUniqueID(), GUE.Structs.SoftObjectPath, Ctx);
-		Out.AddStruct(GUE.Members.Id, GUE.Structs.SoftObjectPath, SoftPath);
+		return Save(Out, In.ToSoftObjectPath(), Ctx);
 	}
 
-	static void Load(FSoftObjectPtr& Out, FStructLoadView In)
+	void Load(FSoftObjectPtr& Out, FStructLoadView In)
 	{
-		Out.ResetWeakPtr();
-		LoadSoleStruct(&Out.GetUniqueID(), In);
+		FSoftObjectPath Temp;
+		Load(Temp, In);
+		Out = Temp;
 	}
 
 	static FSoleMemberSpec Spec(FLazyObjectPtr*, FDeclId Id)
@@ -1197,6 +1229,18 @@ private:
 	TArray<FText> Texts; // Tricky to serialize non-intrusively
 };
 
+struct FLinkerPropertyBindings : TCustomPropertyBindings<FLinkerPropertyBatch>
+{
+	FLinkerPropertyBindings(FLinkerPropertyBatch& Batch, const FCustomBindings& Underlay)
+	: TCustomPropertyBindings<FLinkerPropertyBatch>(Batch, Underlay)
+	, SoftObjectPath(Batch)
+	{
+		Bind(GUE.Structs.SoftObjectPath, SoftObjectPath);
+	}
+	
+	TCustomPropertyBinding<FSoftObjectPath, FLinkerPropertyBatch> SoftObjectPath;
+};
+
 struct FPlainPropsLinkerSandbox
 {
 	static constexpr uint32 LoadFlags = (LOAD_Async | LOAD_NoVerify | LOAD_SkipLoadImportedPackages);
@@ -1251,7 +1295,7 @@ struct FPlainPropsLinkerSandbox
 		FLinkerLoadInstance Base;
 		FLinkerLoadInstance PP;
 		FLinkerPropertyBatch Properties;
-		TCustomPropertyBindings<FLinkerPropertyBatch> Customs;
+		FLinkerPropertyBindings Customs;
 		FMemoryView BatchView;
 		FPersistentBatchLoader BatchLoader;
 	};
@@ -1278,7 +1322,7 @@ struct FPlainPropsLinkerSandbox
 
 		UE_LOGFMT(LogPlainPropsUObject, Log, " Saving PlainProps for {PackageName}...", PackageName);
 		FLinkerPropertyBatch Properties(LinkerSave);
-		TCustomPropertyBindings<FLinkerPropertyBatch> Customs(Properties, GUE.Customs);
+		FLinkerPropertyBindings Customs(Properties, GUE.Customs);
 		FBatchSaver BatchSaver(Customs.Overlay, LinkerSave.ExportMap.Num());
 		for (const FObjectExport& Export : LinkerSave.ExportMap)
 		{
@@ -1369,7 +1413,7 @@ struct FPlainPropsLinkerSandbox
 		}
 	}
 
-	static FDiffObjectContext DiffPackagesLoadedFromMemory(FName PackageName)
+	static FDiffObjectContext DiffPackagesLoadedFromMemory(FName PackageName, FDiffObjectContext::FIgnoreFilter IgnoreDiffFilter)
 	{
 		PackageName = FoldNameNumber(PackageName);
 		if (TUniquePtr<FMountedPackage>* LinkerBatchPtr = MountedPackages.Find(PackageName))
@@ -1377,6 +1421,7 @@ struct FPlainPropsLinkerSandbox
 			FMountedPackage& LinkerBatch = **LinkerBatchPtr;
 
 			FDiffObjectContext Ctx = { LinkerBatch.Base.Package, LinkerBatch.PP.Package };
+			Ctx.IgnoreDiff = IgnoreDiffFilter;
 			DiffObjects(Ctx);
 			return MoveTemp(Ctx);
 		}
@@ -1447,30 +1492,43 @@ TMap<FName, TUniquePtr<FPlainPropsLinkerSandbox::FMountedPackage>> FPlainPropsLi
 
 ////////////////////////////////////////////////////////////////////////////////
 
+static FString GetSavePackageFilename(UPackage* Package)
+{
+	const FString& Extension = Package->ContainsMap() ? FPackageName::GetMapPackageExtension() : FPackageName::GetAssetPackageExtension();
+	return FPaths::Combine(*FPaths::ProjectSavedDir(), TEXT("PlainProps"), Package->GetName() + Extension);
+}
+
 static bool IsNativeCDO(UObject* Object)
 {
 	return Object->HasAllFlags(RF_ClassDefaultObject) && Object->GetClass()->HasAllClassFlags(CLASS_Native);
 }
 
-int32 RoundtripViaBatch(TConstArrayView<UObject*> Objects, ERoundtrip Options)
+int32 RoundtripViaPlainBatch(TConstArrayView<UObject*> Objects, ERoundtrip Options)
 {
-	TRACE_CPUPROFILER_EVENT_SCOPE(RoundtripViaBatch);
+	TRACE_CPUPROFILER_EVENT_SCOPE(RoundtripViaPlainBatch);
 
 	TArray<FInstance> Instances;
 	if (Objects.Num())
 	{
 		UE_LOGFMT(LogPlainPropsUObject, Display, "Gathering UObjects in Assets...");
-		for (const UObject* Object : Objects)
+		for (UObject* Object : Objects)
 		{
-			ForEachObjectWithPackage(Object->GetPackage(), [&Instances](UObject* Obj)
+			if (Object->IsA<UPackage>())
 			{
-				UClass* Class = Obj->GetClass();
-				if (ShouldBind(Class) && IncludeClass(Class))
+				ForEachObjectWithPackage(Object->GetPackage(), [&Instances](UObject* Obj)
 				{
-					Instances.Emplace(Obj->GetPathName(), Obj);
-				}
-				return true;
-			});
+					UClass* Class = Obj->GetClass();
+					if (ShouldBind(Class) && IncludeClass(Class))
+					{
+						Instances.Emplace(Obj->GetPathName(), Obj);
+					}
+					return true;
+				});
+			}
+			else
+			{
+				Instances.Emplace(Object->GetPathName(), Object);
+			}
 		}
 	}
 	else
@@ -1690,9 +1748,9 @@ int32 RoundtripViaBatch(TConstArrayView<UObject*> Objects, ERoundtrip Options)
 
 ////////////////////////////////////////////////////////////////////////////////
 
-int32 RoundtripViaPackages(TConstArrayView<UObject*> Objects, ERoundtrip Options)
+int32 RoundtripViaLinkerBatch(TConstArrayView<UObject*> Objects, ERoundtrip Options, FLinkerDiffIgnoreFilter DiffIgnoreFilter)
 {
-	TRACE_CPUPROFILER_EVENT_SCOPE(RoundtripViaPackages);
+	TRACE_CPUPROFILER_EVENT_SCOPE(RoundtripViaLinkerBatch);
 
 	TArray<UPackage*> Packages;
 	TArray<UObject*> AllPackages;
@@ -1710,7 +1768,7 @@ int32 RoundtripViaPackages(TConstArrayView<UObject*> Objects, ERoundtrip Options
 		GetObjectsOfClass(UPackage::StaticClass(), AllPackages);
 	}
 
-	// Filter out in-memory and BP packages
+	// Filter out in-memory packages
 	for (UObject* PackageObj : AllPackages)
 	{
 		if (!PackageObj->HasAnyFlags(RF_WasLoaded))
@@ -1722,21 +1780,6 @@ int32 RoundtripViaPackages(TConstArrayView<UObject*> Objects, ERoundtrip Options
 		{
 			continue;
 		}
-		bool bExclude = false;
-		ForEachObjectWithPackage(Package, [&bExclude](UObject* Object)
-		{
-			UClass* Class = Object->GetClass();
-			if (Class->HasAnyClassFlags(CLASS_CompiledFromBlueprint))
-			{
-				bExclude = true;
-				return false;
-			}
-			return true;
-		});
-		if (bExclude)
-		{
-			continue;
-		}
 		Packages.Add(Package);
 	}
 
@@ -1766,11 +1809,8 @@ int32 RoundtripViaPackages(TConstArrayView<UObject*> Objects, ERoundtrip Options
 			{
 				SaveArgs.SaveFlags |= SAVE_Unversioned_Properties;
 			}
-			const FString& Extension = Package->ContainsMap() ? FPackageName::GetMapPackageExtension() : FPackageName::GetAssetPackageExtension();
-			FString TempFilename = FPaths::Combine(*FPaths::ProjectSavedDir(), TEXT("PlainProps"), Package->GetName() + Extension);
-			// FString ActualFilename = FPackageName::LongPackageNameToFilename(Asset->GetPackage()->GetName(), Extension);
-
-			UPackage::Save(Package, Asset, *TempFilename, SaveArgs);
+			FString SaveFilename = GetSavePackageFilename(Package);
+			UPackage::Save(Package, Asset, *SaveFilename, SaveArgs);
 		}
 	}
 
@@ -1823,7 +1863,7 @@ int32 RoundtripViaPackages(TConstArrayView<UObject*> Objects, ERoundtrip Options
 			uint32 NumDiffs = 0;
 			for (UPackage* Package : Packages)
 			{
-				FDiffObjectContext Ctx = FPlainPropsLinkerSandbox::DiffPackagesLoadedFromMemory(Package->GetFName());
+				FDiffObjectContext Ctx = FPlainPropsLinkerSandbox::DiffPackagesLoadedFromMemory(Package->GetFName(), DiffIgnoreFilter);
 				if (Ctx.Diffs.Num())
 				{
 					PrintDiff(Diffs, Ctx);
diff --git a/Engine/Plugins/Experimental/PlainPropsUObject/Source/Private/PlainPropsUObjectBindings.cpp b/Engine/Plugins/Experimental/PlainPropsUObject/Source/Private/PlainPropsUObjectBindings.cpp
--- a/Engine/Plugins/Experimental/PlainPropsUObject/Source/Private/PlainPropsUObjectBindings.cpp
+++ b/Engine/Plugins/Experimental/PlainPropsUObject/Source/Private/PlainPropsUObjectBindings.cpp
@@ -3447,6 +3447,10 @@ bool ShouldBind(const UStruct* Struct)
 			return true;
 		}
 	}
+	if (Struct->HasAnyCastFlags(CASTCLASS_UScriptStruct))
+	{
+		return static_cast<const UScriptStruct*>(Struct)->StructFlags & STRUCT_SerializeNative;
+	}
 	return false;
 }
 
@@ -3528,7 +3532,7 @@ void BindStruct(FBindId Id, const UStruct* Struct)
 	}
 }
 
-EPropertyFlags GetPropertiesToSkip(EBindMode Mode)
+static EPropertyFlags GetPropertiesToSkip(EBindMode Mode)
 {
 	switch (Mode)
 	{
@@ -3538,8 +3542,36 @@ EPropertyFlags GetPropertiesToSkip(EBindMode Mode)
 	}
 }
 
-void SchemaBindAllTypes(EBindMode Mode)
+static TArray<FBindId> GetStructsToSkip(EBatchType BatchType)
 {
+	switch (BatchType)
+	{
+	case EBatchType::Plain:		return { GUE.Structs.VerseFunction };
+	case EBatchType::Linker:	return { GUE.Structs.VerseFunction, GUE.Structs.SoftObjectPath };
+	default:					return {};
+	}
+}
+
+static void InitBatchedProperties(EBatchType BatchType)
+{
+	GUE.Defaults.BindZeroes(GUE.Structs.FieldPath, sizeof(FFieldPath), alignof(FFieldPath));
+	GUE.Defaults.BindZeroes(GUE.Structs.Name, sizeof(FName), alignof(FName));
+	GUE.Defaults.BindStatic(GUE.Structs.Text, &FText::GetEmpty());
+	GUE.Defaults.BindZeroes(GUE.Structs.ClassPtr, sizeof(TSubclassOf<UClass>), alignof(TSubclassOf<UClass>));
+	GUE.Defaults.BindZeroes(GUE.Structs.ObjectPtr, sizeof(FObjectPtr), alignof(FObjectPtr));
+	GUE.Defaults.BindZeroes(GUE.Structs.WeakObjectPtr, sizeof(FWeakObjectPtr), alignof(FWeakObjectPtr));
+	GUE.Defaults.BindZeroes(GUE.Structs.SoftObjectPtr, sizeof(FSoftObjectPtr), alignof(FSoftObjectPtr));
+	GUE.Defaults.BindZeroes(GUE.Structs.LazyObjectPtr, sizeof(FLazyObjectPtr), alignof(FLazyObjectPtr));
+	if (BatchType == EBatchType::Linker)
+	{
+		GUE.Defaults.BindZeroes(GUE.Structs.SoftObjectPath, sizeof(FSoftObjectPath), alignof(FSoftObjectPath));
+	}
+}
+
+void SchemaBindAllTypes(EBindMode Mode, EBatchType BatchType)
+{
+	InitBatchedProperties(BatchType);
+
 	GDontBind = GetPropertiesToSkip(Mode);
 	
 	// Declare all UEnums
@@ -3549,7 +3581,7 @@ void SchemaBindAllTypes(EBindMode Mode)
 	}
 
 	// Bind all UScriptStruct/UClass/UFunction
-	static const FBindId SkipStructs[] = { GUE.Structs.VerseFunction };
+	TArray<FBindId> SkipStructs = GetStructsToSkip(BatchType);
 	for (TObjectIterator<UStruct> It; It; ++It)
 	{
 		const UStruct* Struct = *It;
@@ -3563,18 +3595,6 @@ void SchemaBindAllTypes(EBindMode Mode)
 
 ////////////////////////////////////////////////////////////////////////////////////////////////
 
-static void InitBatchedProperties()
-{
-	GUE.Defaults.BindZeroes(GUE.Structs.FieldPath, sizeof(FFieldPath), alignof(FFieldPath));
-	GUE.Defaults.BindZeroes(GUE.Structs.Name, sizeof(FName), alignof(FName));
-	GUE.Defaults.BindStatic(GUE.Structs.Text, &FText::GetEmpty());
-	GUE.Defaults.BindZeroes(GUE.Structs.ClassPtr, sizeof(TSubclassOf<UClass>), alignof(TSubclassOf<UClass>));
-	GUE.Defaults.BindZeroes(GUE.Structs.ObjectPtr, sizeof(FObjectPtr), alignof(FObjectPtr));
-	GUE.Defaults.BindZeroes(GUE.Structs.WeakObjectPtr, sizeof(FWeakObjectPtr), alignof(FWeakObjectPtr));
-	GUE.Defaults.BindZeroes(GUE.Structs.SoftObjectPtr, sizeof(FSoftObjectPtr), alignof(FSoftObjectPtr));
-	GUE.Defaults.BindZeroes(GUE.Structs.LazyObjectPtr, sizeof(FLazyObjectPtr), alignof(FLazyObjectPtr));
-}
-
 struct FPlainPropsBindingScope
 {
 	inline static const FTransform DefaultTransform;
@@ -3632,7 +3652,6 @@ void FPlainPropsUObjectModule::StartupModule()
 	using namespace PlainProps::UE;
 	DebugIdScopePtr = MakeUnique<DbgVis::FIdScope>(GUE.Names, "SensName");
 	BindingScopePtr = MakeUnique<FPlainPropsBindingScope>();
-	InitBatchedProperties();
 }
 
 void FPlainPropsUObjectModule::ShutdownModule()
diff --git a/Engine/Plugins/Experimental/PlainPropsUObject/Source/Private/PlainPropsUObjectDiff.cpp b/Engine/Plugins/Experimental/PlainPropsUObject/Source/Private/PlainPropsUObjectDiff.cpp
--- a/Engine/Plugins/Experimental/PlainPropsUObject/Source/Private/PlainPropsUObjectDiff.cpp
+++ b/Engine/Plugins/Experimental/PlainPropsUObject/Source/Private/PlainPropsUObjectDiff.cpp
@@ -108,6 +108,41 @@ bool DiffObjectReferences(const UObject* A, const UObject* B, FDiffObjectContext
 	return DiffStructs(A, B, Class, Ctx);
 }
 
+bool DiffSoftObjectPaths(const FSoftObjectPath* A, const FSoftObjectPath* B, const FStructProperty* Property, FDiffObjectContext& Ctx)
+{
+	if (A == B)
+	{
+		return false;
+	}
+
+	if (!A || !B)
+	{
+		Ctx.Diffs.Emplace(EDiffObjectNodeType::Property, A, B, Property);
+		return true;
+	}
+
+	bool bDiff;
+	bool bIsInA = A->GetLongPackageFName() == Ctx.RootA->GetFName();
+	bool bIsInB = B->GetLongPackageFName() == Ctx.RootB->GetFName();
+	if (bIsInA & bIsInB)
+	{
+		bDiff = A->GetAssetFName() != B->GetAssetFName() || A->GetSubPathUtf8String() != B->GetSubPathUtf8String();
+	}
+	else if (bIsInA | bIsInB)
+	{
+		bDiff = true;
+	}
+	else
+	{
+		bDiff = *A != *B;
+	}
+	if (bDiff)
+	{
+		Ctx.Diffs.Emplace(EDiffObjectNodeType::Property, A, B, Property);
+	}
+	return bDiff;
+}
+
 bool DiffObjects(const UObject* A, const UObject* B, FDiffObjectContext& Ctx)
 {
 	if (A == B)
@@ -145,14 +180,30 @@ bool DiffObjects(const UObject* A, const UObject* B, FDiffObjectContext& Ctx)
 
 bool DiffStructProperties(const void* A, const void* B, const FStructProperty* Property, FDiffObjectContext& Ctx)
 {
+	static const FName SoftObjectPath{"SoftObjectPath"};
+
 	const UScriptStruct* ScriptStruct = Property->Struct;
+	if (ScriptStruct->GetFName() == SoftObjectPath)
+	{
+		return DiffSoftObjectPaths(static_cast<const FSoftObjectPath*>(A), static_cast<const FSoftObjectPath*>(B), Property, Ctx);
+	}
+
+	// see ScriptStruct->CompareScriptStruct(A, B, /*PortFlags*/ 0);
 	if (ScriptStruct->StructFlags & STRUCT_IdenticalNative)
 	{
-		bool bOut;
-		if (!ScriptStruct->GetCppStructOps()->Identical(A, B, /*PortFlags*/0, bOut))
+		bool bResult;
+		if (ScriptStruct->GetCppStructOps()->Identical(A, B, /*PortFlags*/0, bResult))
 		{
-			Ctx.Diffs.Emplace(EDiffObjectNodeType::Property, A, B, Property);
-			return true;
+			if (!bResult)
+			{
+				// place breakpoint here
+				if (Ctx.IgnoreDiff(A, B, Property))
+				{
+					return false;
+				}
+				Ctx.Diffs.Emplace(EDiffObjectNodeType::Property, A, B, Property);
+			}
+			return !bResult;
 		}
 	}
 	if (DiffStructs(A, B, ScriptStruct, Ctx))
@@ -260,10 +311,10 @@ bool DiffMapProperties(const void* A, const void* B, const FMapProperty* Propert
 
 bool DiffInterfaceProperties(const void* A, const void* B, const FInterfaceProperty* Property, FDiffObjectContext& Ctx)
 {
-	FScriptInterface* InterfaceA = (FScriptInterface*)A;
-	FScriptInterface* InterfaceB = (FScriptInterface*)B;
-	UObject* ObjectA = InterfaceA->GetObject();
-	UObject* ObjectB = InterfaceB->GetObject();
+	const FScriptInterface* InterfaceA = (FScriptInterface*)A;
+	const FScriptInterface* InterfaceB = (FScriptInterface*)B;
+	const UObject* ObjectA = InterfaceA->GetObject();
+	const UObject* ObjectB = InterfaceB->GetObject();
 	if (DiffObjectReferences(ObjectA, ObjectB, Ctx))
 	{
 		Ctx.Diffs.Emplace(EDiffObjectNodeType::Property, A, B, Property);
@@ -272,6 +323,19 @@ bool DiffInterfaceProperties(const void* A, const void* B, const FInterfacePrope
 	return false;
 }
 
+bool DiffDelegateProperties(const void* A, const void* B, const FDelegateProperty* Property, FDiffObjectContext& Ctx)
+{
+	const FScriptDelegate* DelegateA = (const FScriptDelegate*)A;
+	const FScriptDelegate* DelegateB = (const FScriptDelegate*)B;
+	if (DelegateA->GetFunctionName() != DelegateB->GetFunctionName() ||
+		DiffObjectReferences(DelegateA->GetUObject(), DelegateB->GetUObject(), Ctx))
+	{
+		Ctx.Diffs.Emplace(EDiffObjectNodeType::Property, A, B, Property);
+		return true;
+	}
+	return false;
+}
+
 bool DiffProperties(const void* A, const void* B, const FProperty* Property, FDiffObjectContext& Ctx)
 {
 	uint64 Flags = Property->GetCastFlags();
@@ -318,14 +382,22 @@ bool DiffProperties(const void* A, const void* B, const FProperty* Property, FDi
 			return true;
 		}
 	}
-	else if (HasAnyCastClassFlags<CASTCLASS_FTextProperty>(Flags))
+	else if (HasAnyCastClassFlags<CASTCLASS_FDelegateProperty>(Flags))
 	{
-		return false;
+		if (DiffDelegateProperties(A, B, static_cast<const FDelegateProperty*>(Property), Ctx))
+		{
+			return true;
+		}
 	}
 	else
 	{
 		if (!Property->Identical(A, B, PPF_DeepComparison))
 		{
+			// place breakpoint here
+			if (Ctx.IgnoreDiff(A, B, Property))
+			{
+				return false;
+			}
 			Ctx.Diffs.Emplace(EDiffObjectNodeType::Property, A, B, Property);
 			return true;
 		}
@@ -335,12 +407,19 @@ bool DiffProperties(const void* A, const void* B, const FProperty* Property, FDi
 
 bool DiffStructs(const void* A, const void* B, const UStruct* Struct, FDiffObjectContext& Ctx)
 {
+	const static FName ClassDefaultObject{"ClassDefaultObject"};
+	const bool bIsUClass = Struct->IsChildOf<UClass>();
 	for (FProperty* Property = Struct->PropertyLink; Property; Property = Property->PropertyLinkNext)
 	{
 		if (!Property->ShouldDuplicateValue()) //Should the property be compared at all?
 		{
 			continue;
 		}
+		// don't compare the CDO property for classes to prevent infinite recursion
+		if (bIsUClass && Property->GetName() == ClassDefaultObject)
+		{
+			continue;
+		}
 		const void* ValueA = Property->ContainerPtrToValuePtr<void>(A);
 		const void* ValueB = Property->ContainerPtrToValuePtr<void>(B);
 		if (DiffProperties(ValueA, ValueB, Property, Ctx))
@@ -417,6 +496,8 @@ void PrintDiff(FUtf8StringBuilderBase& Out, FDiffObjectContext& Ctx)
 		FString B;
 		Property->ExportTextItem_Direct(A, LeafNode.A, nullptr, (UObject*)Ctx.RootA, PPF_None, nullptr);
 		Property->ExportTextItem_Direct(B, LeafNode.B, nullptr, (UObject*)Ctx.RootB, PPF_None, nullptr);
+		A.ReplaceInline(TEXT("\r\n"), TEXT(" "));
+		B.ReplaceInline(TEXT("\r\n"), TEXT(" "));
 		Out << A << "' != '" << B;
 	}
 	else
diff --git a/Engine/Plugins/Experimental/PlainPropsUObject/Source/Private/PlainPropsUObjectDiffInternal.h b/Engine/Plugins/Experimental/PlainPropsUObject/Source/Private/PlainPropsUObjectDiffInternal.h
--- a/Engine/Plugins/Experimental/PlainPropsUObject/Source/Private/PlainPropsUObjectDiffInternal.h
+++ b/Engine/Plugins/Experimental/PlainPropsUObject/Source/Private/PlainPropsUObjectDiffInternal.h
@@ -1,9 +1,15 @@
 // Copyright Epic Games, Inc. All Rights Reserved.
 
+// Temporary custom UObject diffing code that allows for comparing two object graphs by treating
+// internal object references as equal if their relative object paths are equal.
+// This is currently only used by RoundtripViaLinkerBatch which will be updated to use a more robust
+// and scalable approach that also captures custom bound data.
+
 #pragma once
 
 #include "Containers/Array.h"
 #include "Containers/StringFwd.h"
+#include "Templates/Function.h"
 
 class FProperty;
 class UClass;
@@ -47,9 +53,12 @@ struct FDiffObjectNode
 
 struct FDiffObjectContext
 {
+	using FIgnoreFilter = TFunction<bool(const void* A, const void* B, const FProperty*)>;
+
 	// in
 	const UObject* const RootA = nullptr;
 	const UObject* const RootB = nullptr;
+	FIgnoreFilter IgnoreDiff;
 
 	// out
 	TArray<FDiffObjectNode> Diffs;
diff --git a/Engine/Plugins/Experimental/PlainPropsUObject/Source/Public/PlainPropsRoundtripTest.h b/Engine/Plugins/Experimental/PlainPropsUObject/Source/Public/PlainPropsRoundtripTest.h
--- a/Engine/Plugins/Experimental/PlainPropsUObject/Source/Public/PlainPropsRoundtripTest.h
+++ b/Engine/Plugins/Experimental/PlainPropsUObject/Source/Public/PlainPropsRoundtripTest.h
@@ -4,6 +4,7 @@
 
 #include "Containers/ArrayView.h"
 #include "Misc/EnumClassFlags.h"
+#include "Templates/Function.h"
 
 class UObject;
 
@@ -12,6 +13,8 @@ namespace PlainProps::UE
 
 enum class EBindMode : uint8 { All, Source, Runtime };
 
+enum class EBatchType : uint8 { Plain, Linker} ; // todo: Iostore, Zen
+
 enum class ERoundtrip : uint8
 {
 	None 		= 0,
@@ -23,8 +26,10 @@ enum class ERoundtrip : uint8
 };
 ENUM_CLASS_FLAGS(ERoundtrip);
 
-PLAINPROPSUOBJECT_API void SchemaBindAllTypes(EBindMode Mode);
-PLAINPROPSUOBJECT_API int32 RoundtripViaBatch(TConstArrayView<UObject*> Objects, ERoundtrip Options);
-PLAINPROPSUOBJECT_API int32 RoundtripViaPackages(TConstArrayView<UObject*> Objects, ERoundtrip Options);
+using FLinkerDiffIgnoreFilter = TFunction<bool(const void* A, const void* B, const FProperty*)>;
+
+PLAINPROPSUOBJECT_API void SchemaBindAllTypes(EBindMode Mode, EBatchType BatchType);
+PLAINPROPSUOBJECT_API int32 RoundtripViaPlainBatch(TConstArrayView<UObject*> Objects, ERoundtrip Options);
+PLAINPROPSUOBJECT_API int32 RoundtripViaLinkerBatch(TConstArrayView<UObject*> Objects, ERoundtrip Options, FLinkerDiffIgnoreFilter DiffIgnoreFilter);
 
 }
