Added additional logic to handle the serialization of CDOs in PlainPropsRoundtripTest.
Using an object's archetype instead of the base instance when serializing with TPS and UPS.

#rnx

diff --git a/Engine/Plugins/Experimental/PlainPropsUObject/Source/Private/PlainPropsRoundtripTest.cpp b/Engine/Plugins/Experimental/PlainPropsUObject/Source/Private/PlainPropsRoundtripTest.cpp
--- a/Engine/Plugins/Experimental/PlainPropsUObject/Source/Private/PlainPropsRoundtripTest.cpp
+++ b/Engine/Plugins/Experimental/PlainPropsUObject/Source/Private/PlainPropsRoundtripTest.cpp
@@ -599,7 +599,9 @@ public:
 	FORCENOINLINE void WriteProperties(UObject* Object, UObject* Defaults)
 	{
 		UClass* Class = Object->GetClass();
-		Class->SerializeTaggedProperties(*this, reinterpret_cast<uint8*>(Object), Class, reinterpret_cast<uint8*>(Defaults));
+		// See DiffClass in UObject::SerializeScriptProperties
+		UClass* DefaultsClass = Object->HasAnyFlags(RF_ClassDefaultObject) ? Class->GetSuperClass() : Class;
+		Class->SerializeTaggedProperties(*this, reinterpret_cast<uint8*>(Object), DefaultsClass, reinterpret_cast<uint8*>(Defaults));
 	}
 
 	virtual FArchive& operator<<(FText& Value) override
@@ -642,7 +644,9 @@ public:
 	FORCENOINLINE void ReadProperties(UObject* Object, UObject* Defaults)
 	{
 		UClass* Class = Object->GetClass();
-		Class->SerializeTaggedProperties(*this, reinterpret_cast<uint8*>(Object), Class, reinterpret_cast<uint8*>(Defaults));
+		// See DiffClass in UObject::SerializeScriptProperties
+		UClass* DefaultsClass = Object->HasAnyFlags(RF_ClassDefaultObject) ? Class->GetSuperClass() : Class;
+		Class->SerializeTaggedProperties(*this, reinterpret_cast<uint8*>(Object), DefaultsClass, reinterpret_cast<uint8*>(Defaults));
 	}
 
 	virtual FArchive& operator<<(FText& Value) override
@@ -713,7 +717,16 @@ static UObject* MakeEmptyInstance(UObject* Obj, FName Name)
 	Params.Template = Obj->GetArchetype();
 	Params.bAssumeTemplateIsArchetype = true;
 	Params.bCopyTransientsFromClassDefaults = !(Params.SetFlags & RF_ClassDefaultObject);
-	return StaticConstructObject_Internal(Params);
+	UObject* Result = StaticConstructObject_Internal(Params);
+
+	// When creating CDOs the name passed-in will be ignored and the default name will be used instead.
+	// (e.g., Default__MyType). To avoid reusing the same instance we force a rename.
+	if (Result->GetFName() != Name)
+	{
+		Result->Rename(*Name.ToString(), nullptr, REN_DontCreateRedirectors | REN_DoNotDirty | REN_NonTransactional);
+	}
+
+	return Result;
 }
 
 static bool IncludeClass(UClass* Class)
@@ -773,7 +786,7 @@ FORCENOINLINE void SaveArchive(FPropertyWriter& Archive, TConstArrayView<FInstan
 	Archive.SetUseUnversionedPropertySerialization(bUps);
 	for (const FInstance& Instance : Instances)
 	{
-		Archive.WriteProperties(Instance.Orig, Instance.Base);
+		Archive.WriteProperties(Instance.Orig, Instance.Arch);
 	}
 }
 
@@ -783,7 +796,7 @@ FORCENOINLINE void LoadArchive(FPropertyReader& Archive, TConstArrayView<FInstan
 	Archive.SetUseUnversionedPropertySerialization(bUps);
 	for (const FInstance& Instance : Instances)
 	{
-		Archive.ReadProperties(bUps ? Instance.UPS : Instance.TPS, Instance.Base);
+		Archive.ReadProperties(bUps ? Instance.UPS : Instance.TPS, Instance.Arch);
 	}
 }
 
@@ -1576,6 +1589,10 @@ int32 RoundtripViaPlainBatch(TConstArrayView<UObject*> Objects, ERoundtrip Optio
 		{
 			Instance.UPS = MakeEmptyInstance(Instance.Orig, FName("UPS", N));
 		}
+
+		check(!Instance.PP  || (Instance.PP  != Instance.Base));
+		check(!Instance.TPS || (Instance.TPS != Instance.Base));
+		check(!Instance.UPS || (Instance.UPS != Instance.Base));
 	}
 
 	// Save
