Added PlainProps serialization as another test permutation in the overridable serialization unit tests.
Note that this change adds the PlainProps test code in a disabled state without adding the necessary dependencies.

#rnx
#Synced-CL 48228706

diff --git a/Engine/Plugins/Experimental/PlainPropsEngine/Source/Private/PlainPropsEngineTestUtils.cpp b/Engine/Plugins/Experimental/PlainPropsEngine/Source/Private/PlainPropsEngineTestUtils.cpp
new file mode 100644
--- /dev/null
+++ b/Engine/Plugins/Experimental/PlainPropsEngine/Source/Private/PlainPropsEngineTestUtils.cpp
@@ -0,0 +1,18 @@
+// Copyright Epic Games, Inc. All Rights Reserved.
+
+#include "PlainPropsEngineTestUtils.h"
+#include "PlainPropsEngineBindings.h"
+#include "PlainPropsRoundtripTest.h"
+
+namespace PlainProps::UE
+{
+
+void BindAllTypes(EBindMode Mode, EBatchType BatchType)
+{
+	SchemaBindAllTypes(Mode, BatchType);
+	CustomBindEngineTypes(Mode);
+	CustomBindAnimationTypes(Mode);
+	CustomBindNiagaraTypes(Mode);
+}
+
+} // namespace PlainProps::UE
\ No newline at end of file
diff --git a/Engine/Plugins/Experimental/PlainPropsEngine/Source/Public/PlainPropsEngineTestUtils.h b/Engine/Plugins/Experimental/PlainPropsEngine/Source/Public/PlainPropsEngineTestUtils.h
new file mode 100644
--- /dev/null
+++ b/Engine/Plugins/Experimental/PlainPropsEngine/Source/Public/PlainPropsEngineTestUtils.h
@@ -0,0 +1,14 @@
+// Copyright Epic Games, Inc. All Rights Reserved.
+
+#pragma once
+
+#include "HAL/Platform.h"
+
+namespace PlainProps::UE
+{
+enum class EBindMode : uint8;
+enum class EBatchType : uint8;
+
+PLAINPROPSENGINE_API void BindAllTypes(EBindMode Mode, EBatchType BatchType);
+
+}
\ No newline at end of file
diff --git a/Engine/Plugins/Experimental/PlainPropsUObject/Source/Private/PlainPropsRoundtripTest.cpp b/Engine/Plugins/Experimental/PlainPropsUObject/Source/Private/PlainPropsRoundtripTest.cpp
--- a/Engine/Plugins/Experimental/PlainPropsUObject/Source/Private/PlainPropsRoundtripTest.cpp
+++ b/Engine/Plugins/Experimental/PlainPropsUObject/Source/Private/PlainPropsRoundtripTest.cpp
@@ -1404,7 +1404,7 @@ struct FPlainPropsLinkerSandbox
 		}
 	}
 
-	static void LoadPlainPropsPackageFromMemory(FName PackageName)
+	static UPackage* LoadPlainPropsPackageFromMemory(FName PackageName)
 	{
 		PackageName = FoldNameNumber(PackageName);
 		if (TUniquePtr<FMountedPackage>* LinkerBatchPtr = MountedPackages.Find(PackageName))
@@ -1412,8 +1412,11 @@ struct FPlainPropsLinkerSandbox
 			FMountedPackage& LinkerBatch = **LinkerBatchPtr;
 
 			UE_LOGFMT(LogPlainPropsUObject, Log, " Loading PlainProps for {Package}...", PackageName);
-			LoadPackage(LinkerBatch.PP.Package, LinkerBatch.PackagePath, LoadFlags);
+
+			return LoadPackage(LinkerBatch.PP.Package, LinkerBatch.PackagePath, LoadFlags);
 		}
+
+		return nullptr;
 	}
 
 	static void RoundtripToTextFromMemory(FName PackageName)
@@ -1765,6 +1768,52 @@ int32 RoundtripViaPlainBatch(TConstArrayView<UObject*> Objects, ERoundtrip Optio
 
 ////////////////////////////////////////////////////////////////////////////////
 
+#if PP_WITH_ENGINE_HOOKS
+struct FScopeSaveTestPackage
+{
+	FScopeSaveTestPackage()
+	{
+		FSavePackageSettings& SavePackageSettings = FSavePackageSettings::GetDefaultSettings();
+		ExternalWriteExportsFunc = SavePackageSettings.ExternalWriteExportsFunc;
+		ExternalWriteFileFunc = SavePackageSettings.ExternalWriteFileFunc;
+		SerializePropertiesFunc = GSerializePropertiesOverride;
+		LoadFromPackageFileFunc = GLoadFromPackageFileOverride;
+
+		SavePackageSettings.ExternalWriteExportsFunc = &FPlainPropsLinkerSandbox::SavePlainProps;
+		SavePackageSettings.ExternalWriteFileFunc = &FPlainPropsLinkerSandbox::SavePackageToMemory;
+		GSerializePropertiesOverride = &FPlainPropsLinkerSandbox::LoadObject;
+		GLoadFromPackageFileOverride = &FPlainPropsLinkerSandbox::OpenPackageReader;
+	}
+
+	~FScopeSaveTestPackage()
+	{
+		FSavePackageSettings& SavePackageSettings = FSavePackageSettings::GetDefaultSettings();
+		SavePackageSettings.ExternalWriteExportsFunc = ExternalWriteExportsFunc;
+		SavePackageSettings.ExternalWriteFileFunc = ExternalWriteFileFunc;
+		GSerializePropertiesOverride = SerializePropertiesFunc;
+		GLoadFromPackageFileOverride = LoadFromPackageFileFunc;
+	}
+
+	void(*ExternalWriteExportsFunc)(const FLinkerSave& LinkerSave) = nullptr;
+	void(*ExternalWriteFileFunc)(const FLinkerSave& LinkerSave, uint8* Data, int64 Size) = nullptr;
+	bool(*SerializePropertiesFunc)(class FStructuredArchiveSlot&, uint8*);
+	bool (*LoadFromPackageFileFunc)(const FPackagePath&, FArchive*&) = nullptr;
+};
+
+// For now the 2 scopes are identical but we will eventually want to have more fine-grained
+// control over what gets overwritten between save and load.
+struct FScopeLoadTestPackage : public FScopeSaveTestPackage
+{
+};
+#else
+struct FScopeSaveTestPackage 
+{
+	FScopeSaveTestPackage() {}
+	~FScopeSaveTestPackage() {}
+};
+struct FScopeLoadTestPackage : public FScopeSaveTestPackage {};
+#endif // PP_WITH_ENGINE_HOOKS
+
 int32 RoundtripViaLinkerBatch(TConstArrayView<UObject*> Objects, ERoundtrip Options, FLinkerDiffIgnoreFilter DiffIgnoreFilter)
 {
 	TRACE_CPUPROFILER_EVENT_SCOPE(RoundtripViaLinkerBatch);
@@ -1801,12 +1850,6 @@ int32 RoundtripViaLinkerBatch(TConstArrayView<UObject*> Objects, ERoundtrip Opti
 	}
 
 #if PP_WITH_ENGINE_HOOKS
-	// Todo consider TGuardValue
-	FSavePackageSettings& SavePackageSettings = FSavePackageSettings::GetDefaultSettings();
-	SavePackageSettings.ExternalWriteExportsFunc = &FPlainPropsLinkerSandbox::SavePlainProps;
-	SavePackageSettings.ExternalWriteFileFunc = &FPlainPropsLinkerSandbox::SavePackageToMemory;
-	GSerializePropertiesOverride = &FPlainPropsLinkerSandbox::LoadObject;
-	GLoadFromPackageFileOverride = &FPlainPropsLinkerSandbox::OpenPackageReader;
 	FPlainPropsLinkerSandbox::SavedPackages.Reserve(Packages.Num());
 	FPlainPropsLinkerSandbox::MountedPackages.Reserve(Packages.Num());
 #endif
@@ -1814,6 +1857,8 @@ int32 RoundtripViaLinkerBatch(TConstArrayView<UObject*> Objects, ERoundtrip Opti
 	UE_LOGFMT(LogPlainPropsUObject, Display, "Saving {SaveNum} packages...", Packages.Num());
 
 	{
+		FScopeSaveTestPackage SaveScope;
+
 		TRACE_CPUPROFILER_EVENT_SCOPE(SavePackages);
 		for (UPackage* Package : Packages)
 		{
@@ -1835,6 +1880,8 @@ int32 RoundtripViaLinkerBatch(TConstArrayView<UObject*> Objects, ERoundtrip Opti
 	// Roundtrip to text, also dumping all packages as yaml to disk
 	if (EnumHasAnyFlags(Options, ERoundtrip::TextMemory | ERoundtrip::TextStable))
 	{
+		FScopeSaveTestPackage SaveScope;
+
 		TRACE_CPUPROFILER_EVENT_SCOPE(RoundtripToText);
 		UE_LOGFMT(LogPlainPropsUObject, Display, "Roundtripping {SaveNum} packages to text...", Packages.Num());
 		for (UPackage* Package : Packages)
@@ -1846,6 +1893,8 @@ int32 RoundtripViaLinkerBatch(TConstArrayView<UObject*> Objects, ERoundtrip Opti
 	// Roundtrip plain props and diff against baseline (TPS or UPS)
 	if (EnumHasAnyFlags(Options, ERoundtrip::PP))
 	{
+		FScopeLoadTestPackage LoadScope;
+
 		{
 			UE_LOGFMT(LogPlainPropsUObject, Display, "Mounting and loading {SaveNum} packages twice...", Packages.Num());
 			for (UPackage* Package : Packages)
@@ -1899,12 +1948,66 @@ int32 RoundtripViaLinkerBatch(TConstArrayView<UObject*> Objects, ERoundtrip Opti
 	// It will reset and flush all the linkers which takes ~15s for ~3000 packages in Release config
 	FPlainPropsLinkerSandbox::Shutdown();
 #endif
-	SavePackageSettings.ExternalWriteExportsFunc = nullptr;
-	GSerializePropertiesOverride = nullptr;
-	GLoadFromPackageFileOverride = nullptr;
-#endif
+
+#endif // PP_WITH_ENGINE_HOOKS
 
 	return 0;
 }
 
+bool SaveTestPackage(UPackage* Package, const FString& Filename, const FSavePackageArgs& SaveArgs)
+{
+	bool bDidSave = false;
+
+#if PP_WITH_ENGINE_HOOKS
+	check(Package);
+	checkf(SaveArgs.SaveFlags & SAVE_Async, TEXT("Only packages saved to a temporary memory writer are currently supported"));
+
+	FScopeSaveTestPackage SaveScope;
+
+	// We need a valid path to be set on the package for when we load it back.
+	FPackagePath PackagePath;
+	PackagePath = FPackagePath::FromLocalPath(Filename);
+	Package->SetLoadedPath(PackagePath);
+
+	const FSavePackageResultStruct SavePackageResult = UPackage::Save(Package, nullptr, *Filename, SaveArgs);
+
+	IFileManager& FileManager = IFileManager::Get();
+
+	if (SavePackageResult.IsSuccessful())
+	{
+		// Because we have overridden ExternalWriteFileFunc to only write to memory
+		// (see FScopeSaveTestPackage) the previous save did not generate an actual file.
+		// This will cause our package to be considered missing during load (see
+		// TryGetPackageInfoFromPackageStore). As a workaround we create a temporary
+		// empty package file.
+		if (FArchive* Ar = FileManager.CreateFileWriter(*Filename))
+		{
+			uint32 TempData = 0;
+			Ar->Serialize(&TempData, sizeof(TempData));
+			delete Ar;
+		}
+	}
+
+	bDidSave = SavePackageResult.IsSuccessful() && FileManager.FileExists(*Filename);
+#endif
+
+	return bDidSave;
+}
+
+UPackage* LoadTestPackage(const FString& Filename)
+{
+	UPackage* Result = nullptr;
+
+#if PP_WITH_ENGINE_HOOKS
+	FScopeLoadTestPackage LoadScope;
+
+	FName PackageName = FName(*Filename);
+
+	FPlainPropsLinkerSandbox::MountPackagesFromMemory(PackageName);
+	Result = FPlainPropsLinkerSandbox::LoadPlainPropsPackageFromMemory(PackageName);
+#endif
+
+	return Result;
+}
+
 } // namespace PlainProps::UE
diff --git a/Engine/Plugins/Experimental/PlainPropsUObject/Source/Private/PlainPropsUObjectBindings.cpp b/Engine/Plugins/Experimental/PlainPropsUObject/Source/Private/PlainPropsUObjectBindings.cpp
--- a/Engine/Plugins/Experimental/PlainPropsUObject/Source/Private/PlainPropsUObjectBindings.cpp
+++ b/Engine/Plugins/Experimental/PlainPropsUObject/Source/Private/PlainPropsUObjectBindings.cpp
@@ -295,6 +295,7 @@ FCommonStructIds::FCommonStructIds(const FCommonScopeIds& Scopes, TIdIndexer<FSe
 , VerseFunction(			Names.IndexStruct({Scopes.CoreUObject, Names.MakeTypename("VerseFunction")}))
 , DynamicallyTypedValue(	Names.IndexStruct({Scopes.CoreUObject, Names.MakeTypename("DynamicallyTypedValue")}))
 , ReferencePropertyValue(	Names.IndexStruct({Scopes.CoreUObject, Names.MakeTypename("ReferencePropertyValue")}))
+, PropertyBag(				Names.IndexStruct({Scopes.CoreUObject, Names.MakeTypename("PropertyBag")}))
 {}
 
 FCommonMemberIds::FCommonMemberIds(TIdIndexer<FSensitiveName>& Names)
@@ -3547,7 +3548,7 @@ static TArray<FBindId> GetStructsToSkip(EBatchType BatchType)
 	switch (BatchType)
 	{
 	case EBatchType::Plain:		return { GUE.Structs.VerseFunction };
-	case EBatchType::Linker:	return { GUE.Structs.VerseFunction, GUE.Structs.SoftObjectPath };
+	case EBatchType::Linker:	return { GUE.Structs.VerseFunction, GUE.Structs.SoftObjectPath, GUE.Structs.PropertyBag };
 	default:					return {};
 	}
 }
@@ -3586,10 +3587,29 @@ void SchemaBindAllTypes(EBindMode Mode, EBatchType BatchType)
 	{
 		const UStruct* Struct = *It;
 		FBindId Id = GUE.Names.IndexBindId(IndexType(Struct));
-		if (!Algo::Find(SkipStructs, Id))
+		if (Algo::Find(SkipStructs, Id))
 		{
-			BindStruct(Id, Struct);
+			continue;
 		}
+		
+		// Additional 'IsA' exclusion check (e.g., to exclude UPropertyBag types).
+		bool bSkipStruct = false;
+		for (const UStruct* Base : UStruct::FConstSuperStructIterator(Struct->GetClass()))
+		{
+			FBindId BaseId = GUE.Names.IndexBindId(IndexType(Base));
+			if (Algo::Find(SkipStructs, BaseId))
+			{
+				bSkipStruct = true;
+				break;
+			}
+		}
+
+		if (bSkipStruct)
+		{
+			continue;
+		}
+
+		BindStruct(Id, Struct);
 	}
 }
 
diff --git a/Engine/Plugins/Experimental/PlainPropsUObject/Source/Public/PlainPropsRoundtripTest.h b/Engine/Plugins/Experimental/PlainPropsUObject/Source/Public/PlainPropsRoundtripTest.h
--- a/Engine/Plugins/Experimental/PlainPropsUObject/Source/Public/PlainPropsRoundtripTest.h
+++ b/Engine/Plugins/Experimental/PlainPropsUObject/Source/Public/PlainPropsRoundtripTest.h
@@ -7,6 +7,8 @@
 #include "Templates/Function.h"
 
 class UObject;
+class UPackage;
+struct FSavePackageArgs;
 
 namespace PlainProps::UE
 {
@@ -32,4 +34,7 @@ PLAINPROPSUOBJECT_API void SchemaBindAllTypes(EBindMode Mode, EBatchType BatchTy
 PLAINPROPSUOBJECT_API int32 RoundtripViaPlainBatch(TConstArrayView<UObject*> Objects, ERoundtrip Options);
 PLAINPROPSUOBJECT_API int32 RoundtripViaLinkerBatch(TConstArrayView<UObject*> Objects, ERoundtrip Options, FLinkerDiffIgnoreFilter DiffIgnoreFilter);
 
+PLAINPROPSUOBJECT_API bool SaveTestPackage(UPackage* Package, const FString& Filename, const FSavePackageArgs& SaveArgs);
+PLAINPROPSUOBJECT_API UPackage* LoadTestPackage(const FString& Filename);
+
 }
diff --git a/Engine/Plugins/Experimental/PlainPropsUObject/Source/Public/PlainPropsUObjectRuntime.h b/Engine/Plugins/Experimental/PlainPropsUObject/Source/Public/PlainPropsUObjectRuntime.h
--- a/Engine/Plugins/Experimental/PlainPropsUObject/Source/Public/PlainPropsUObjectRuntime.h
+++ b/Engine/Plugins/Experimental/PlainPropsUObject/Source/Public/PlainPropsUObjectRuntime.h
@@ -103,6 +103,7 @@ struct FCommonStructIds
 	FDualStructId				VerseFunction;
 	FDualStructId				DynamicallyTypedValue;
 	FDualStructId				ReferencePropertyValue;
+	FDualStructId				PropertyBag;
 };
 
 struct FCommonMemberIds
