PlainProps: Add the missing implementation for FScopeLoadTestPackage by moving code from FScopeSaveTestPackage.

#rnx

#Synced-CL 0

diff --git a/Engine/Plugins/Experimental/PlainPropsUObject/Source/Private/PlainPropsRoundtripTest.cpp b/Engine/Plugins/Experimental/PlainPropsUObject/Source/Private/PlainPropsRoundtripTest.cpp
--- a/Engine/Plugins/Experimental/PlainPropsUObject/Source/Private/PlainPropsRoundtripTest.cpp
+++ b/Engine/Plugins/Experimental/PlainPropsUObject/Source/Private/PlainPropsRoundtripTest.cpp
@@ -1776,13 +1776,9 @@ struct FScopeSaveTestPackage
 		FSavePackageSettings& SavePackageSettings = FSavePackageSettings::GetDefaultSettings();
 		ExternalWriteExportsFunc = SavePackageSettings.ExternalWriteExportsFunc;
 		ExternalWriteFileFunc = SavePackageSettings.ExternalWriteFileFunc;
-		SerializePropertiesFunc = GSerializePropertiesOverride;
-		LoadFromPackageFileFunc = GLoadFromPackageFileOverride;
 
 		SavePackageSettings.ExternalWriteExportsFunc = &FPlainPropsLinkerSandbox::SavePlainProps;
 		SavePackageSettings.ExternalWriteFileFunc = &FPlainPropsLinkerSandbox::SavePackageToMemory;
-		GSerializePropertiesOverride = &FPlainPropsLinkerSandbox::LoadObject;
-		GLoadFromPackageFileOverride = &FPlainPropsLinkerSandbox::OpenPackageReader;
 	}
 
 	~FScopeSaveTestPackage()
@@ -1790,28 +1786,39 @@ struct FScopeSaveTestPackage
 		FSavePackageSettings& SavePackageSettings = FSavePackageSettings::GetDefaultSettings();
 		SavePackageSettings.ExternalWriteExportsFunc = ExternalWriteExportsFunc;
 		SavePackageSettings.ExternalWriteFileFunc = ExternalWriteFileFunc;
-		GSerializePropertiesOverride = SerializePropertiesFunc;
-		GLoadFromPackageFileOverride = LoadFromPackageFileFunc;
 	}
 
 	void(*ExternalWriteExportsFunc)(const FLinkerSave& LinkerSave) = nullptr;
 	void(*ExternalWriteFileFunc)(const FLinkerSave& LinkerSave, uint8* Data, int64 Size) = nullptr;
-	bool(*SerializePropertiesFunc)(class FStructuredArchiveSlot&, uint8*);
-	bool (*LoadFromPackageFileFunc)(const FPackagePath&, FArchive*&) = nullptr;
 };
 
-// For now the 2 scopes are identical but we will eventually want to have more fine-grained
-// control over what gets overwritten between save and load.
-struct FScopeLoadTestPackage : public FScopeSaveTestPackage
+struct FScopeLoadTestPackage
 {
+	FScopeLoadTestPackage()
+	{
+		SerializePropertiesFunc = GSerializePropertiesOverride;
+		LoadFromPackageFileFunc = GLoadFromPackageFileOverride;
+		GSerializePropertiesOverride = &FPlainPropsLinkerSandbox::LoadObject;
+		GLoadFromPackageFileOverride = &FPlainPropsLinkerSandbox::OpenPackageReader;
+	}
+
+	~FScopeLoadTestPackage()
+	{
+		GSerializePropertiesOverride = SerializePropertiesFunc;
+		GLoadFromPackageFileOverride = LoadFromPackageFileFunc;
+	}
+
+	bool(*SerializePropertiesFunc)(class FStructuredArchiveSlot&, uint8*);
+	bool (*LoadFromPackageFileFunc)(const FPackagePath&, FArchive*&) = nullptr;
 };
 #else
-struct FScopeSaveTestPackage 
+struct FScopeEmptyBase 
 {
-	FScopeSaveTestPackage() {}
-	~FScopeSaveTestPackage() {}
+	FScopeEmptyBase() {}
+	~FScopeEmptyBase() {}
 };
-struct FScopeLoadTestPackage : public FScopeSaveTestPackage {};
+struct FScopeSaveTestPackage : public FScopeEmptyBase {};
+struct FScopeLoadTestPackage : public FScopeEmptyBase {};
 #endif // PP_WITH_ENGINE_HOOKS
 
 int32 RoundtripViaLinkerBatch(TConstArrayView<UObject*> Objects, ERoundtrip Options, FLinkerDiffIgnoreFilter DiffIgnoreFilter)
@@ -1880,8 +1887,6 @@ int32 RoundtripViaLinkerBatch(TConstArrayView<UObject*> Objects, ERoundtrip Opti
 	// Roundtrip to text, also dumping all packages as yaml to disk
 	if (EnumHasAnyFlags(Options, ERoundtrip::TextMemory | ERoundtrip::TextStable))
 	{
-		FScopeSaveTestPackage SaveScope;
-
 		TRACE_CPUPROFILER_EVENT_SCOPE(RoundtripToText);
 		UE_LOGFMT(LogPlainPropsUObject, Display, "Roundtripping {SaveNum} packages to text...", Packages.Num());
 		for (UPackage* Package : Packages)
