[Core] Fix some compile errors when enabling compact sets

#tests Compiled
#rnx
#Synced-CL 48586332

diff --git a/Engine/Plugins/Experimental/PlainProps/Source/Public/PlainPropsUeCoreBindings.h b/Engine/Plugins/Experimental/PlainProps/Source/Public/PlainPropsUeCoreBindings.h
--- a/Engine/Plugins/Experimental/PlainProps/Source/Public/PlainPropsUeCoreBindings.h
+++ b/Engine/Plugins/Experimental/PlainProps/Source/Public/PlainPropsUeCoreBindings.h
@@ -528,8 +528,8 @@ struct TSetBinding : IItemRangeBinding
 #if UE_USE_COMPACT_SET_AS_DEFAULT
 	virtual void ReadItems(FSaveRangeContext& Ctx) const override
 	{
-		const SetType& Set = Ctx.Request.GetRange<SetType>();
-		Ctx.Items.SetAll(Set.GetData(), Set.Num());
+		const TArrayView<const T> SetData = Ctx.Request.GetRange<SetType>().ArrayView();
+		Ctx.Items.SetAll(SetData.GetData(), SetData.Num());
 	}
 #else
 	virtual void ReadItems(FSaveRangeContext& Ctx) const override
diff --git a/Engine/Plugins/Experimental/PlainPropsUObject/Source/Private/PlainPropsUObjectBindings.cpp b/Engine/Plugins/Experimental/PlainPropsUObject/Source/Private/PlainPropsUObjectBindings.cpp
--- a/Engine/Plugins/Experimental/PlainPropsUObject/Source/Private/PlainPropsUObjectBindings.cpp
+++ b/Engine/Plugins/Experimental/PlainPropsUObject/Source/Private/PlainPropsUObjectBindings.cpp
@@ -1970,18 +1970,32 @@ struct TSetPropertyBinding : IItemRangeBinding, ICustomBinding
 		return Set.FindIndex(Elem, Layout, MakeHashFn(Inner), MakeIdenticalFn(Inner));
 	}
 
+#if UE_USE_COMPACT_SET_AS_DEFAULT
 	inline void RemoveElem(FScriptSet& Set, const void* Elem) const
 	{
 		if (int32 Idx = FindIndex(Set, Elem); Idx != INDEX_NONE)
 		{
-			DestroyElem(Set, Idx);
-#if UE_USE_COMPACT_SET_AS_DEFAULT
-			Set.RemoveAt(Idx, Layout, MakeHashFn(Inner));
+			Set.RemoveAt(Idx, Layout, MakeHashFn(Inner), [this](void* Data)
+				{
+					Inner.DestroyItem(Data);
+				});
+		}
+	}
 #else
+	inline void DestroyElem(FScriptSet& Set, int32 Idx) const
+	{
+		Inner.DestroyItem(Set.GetData(Idx, Layout));
+	}
+
+	inline void RemoveElem(FScriptSet& Set, const void* Elem) const
+	{
+		if (int32 Idx = FindIndex(Set, Elem); Idx != INDEX_NONE)
+		{
+			DestroyElem(Set, Idx);
 			Set.RemoveAt(Idx, Layout);
-#endif
 		}
 	}
+#endif
 	
 	inline bool HasItem(const FScriptSet& Set, const void* Elem) const
 	{
@@ -1992,12 +2006,6 @@ struct TSetPropertyBinding : IItemRangeBinding, ICustomBinding
 	{
 		return HasItem(Set, Elem);
 	}
-
-	inline void DestroyElem(FScriptSet& Set) const requires (bLeaves) {}
-	inline void DestroyElem(FScriptSet& Set, int32 Idx) const
-	{
-		Inner.DestroyItem(Set.GetData(Idx, Layout));
-	}
 	
 	inline void DestroyAll(FScriptSet& Set) const requires (bLeaves) {}
 	inline void DestroyAll(FScriptSet& Set) const
@@ -2450,6 +2458,19 @@ struct TMapPropertyCustomBinding : TMapPropertyItemBinding<FScriptMap, KeyKind,
 		return Out;
 	}
 
+#if UE_USE_COMPACT_SET_AS_DEFAULT
+	inline void RemoveKey(FScriptMap& Map, const void* Key) const
+	{
+		if (int32 Idx = FindKey(Map, Key); Idx != INDEX_NONE)
+		{
+			Map.RemoveAt(Idx, Layout, MakeHashFn(InnerKey), [this](void* Pair)
+				{
+					InnerKey.DestroyItem(Pair);
+					InnerValue.DestroyItem(static_cast<uint8*>(Pair) + Layout.ValueOffset);
+				});
+		}
+	}
+#else
 	inline void DestroyPair(FScriptMap& Map, int32 Idx) const
 	{
 		if (InnerKey.bDestruct || InnerValue.bDestruct)
@@ -2465,13 +2486,10 @@ struct TMapPropertyCustomBinding : TMapPropertyItemBinding<FScriptMap, KeyKind,
 		if (int32 Idx = FindKey(Map, Key); Idx != INDEX_NONE)
 		{
 			DestroyPair(Map, Idx);
-#if UE_USE_COMPACT_SET_AS_DEFAULT
-			Map.RemoveAt(Idx, Layout, MakeHashFn(InnerKey));
-#else
 			Map.RemoveAt(Idx, Layout);
-#endif
 		}
 	}
+#endif
 
 	virtual void SaveCustom(FMemberBuilder& Dst, const void* Src, FBaseline Base, const FSaveContext& Ctx) override
 	{
diff --git a/Engine/Plugins/VirtualProduction/Avalanche/Source/AvalancheMediaEditor/Private/Rundown/ChannelLayerStatus/SAvaRundownChannelLayerStatusList.cpp b/Engine/Plugins/VirtualProduction/Avalanche/Source/AvalancheMediaEditor/Private/Rundown/ChannelLayerStatus/SAvaRundownChannelLayerStatusList.cpp
--- a/Engine/Plugins/VirtualProduction/Avalanche/Source/AvalancheMediaEditor/Private/Rundown/ChannelLayerStatus/SAvaRundownChannelLayerStatusList.cpp
+++ b/Engine/Plugins/VirtualProduction/Avalanche/Source/AvalancheMediaEditor/Private/Rundown/ChannelLayerStatus/SAvaRundownChannelLayerStatusList.cpp
@@ -211,7 +211,7 @@ void SAvaRundownChannelLayerStatusList::RefreshList()
 					// No conflict
 					if (LayerPair.Value.Num() == 1)
 					{
-						const int32 PageId = LayerPair.Value[LayerPair.Value.begin().GetId()];
+						const int32 PageId = *LayerPair.Value.FindArbitraryElement();
 
 						if (const FLinearColor* CurrentColor = ComboPageColors.Find(PageId))
 						{
