PlainProps: Refactor RoundtripViaLinkerBatch flow to load full package graphs using the disk package names instead of loading instanced/cloned/duplicated packages.
Replace FJsonObjectConverter module with JsonObjectGraph module.
Implement DropPackages functionality inspired UPackageTools::ReloadPackages.
DropPackages renames existing loaded packages and objects out of the way for later diffing in order to be able to reload the same disk package name again.
Add editor dependency on UnrealEd needed for the DropPackage implementation that relies on engine types and broadcasts.
Improve roundtrip logic for instanced/cloned packages.
Add temporary in-memory serialization of FText to FLinkerPropertyBatch in order to avoid wide spread side effects of missing texts.
Replace FLinkerDiffIgnoreFilter callback with a FLinkerDiffFilter struct and make all ignore diff logic FName-based.
Add specialized diffing for FFieldPathProperty.
Tweak diff printing to display a nicer string that starts with the property class name.
Fix stack overflow when diffing object references by tracking a visited set of object pairs.

Note: Disable linker roundtripping of blueprints temporarily, this will be re-enabled in a follow-up change to support dynamic load planning.

#rnx

diff --git a/Engine/Plugins/Experimental/PlainPropsEngine/Source/PlainPropsEngine.Build.cs b/Engine/Plugins/Experimental/PlainPropsEngine/Source/PlainPropsEngine.Build.cs
--- a/Engine/Plugins/Experimental/PlainPropsEngine/Source/PlainPropsEngine.Build.cs
+++ b/Engine/Plugins/Experimental/PlainPropsEngine/Source/PlainPropsEngine.Build.cs
@@ -21,5 +21,10 @@ public class PlainPropsEngine : ModuleRules
 				"PlainProps",
 				"PlainPropsUObject"
 		});
+
+		if (Target.bBuildEditor == true)
+		{
+			PrivateDependencyModuleNames.Add("UnrealEd");
+		}
 	}
 }
diff --git a/Engine/Plugins/Experimental/PlainPropsEngine/Source/Private/PlainPropsCommandlets.cpp b/Engine/Plugins/Experimental/PlainPropsEngine/Source/Private/PlainPropsCommandlets.cpp
--- a/Engine/Plugins/Experimental/PlainPropsEngine/Source/Private/PlainPropsCommandlets.cpp
+++ b/Engine/Plugins/Experimental/PlainPropsEngine/Source/Private/PlainPropsCommandlets.cpp
@@ -7,131 +7,180 @@
 #include "AssetRegistry/AssetRegistryModule.h"
 #include "Misc/AsciiSet.h"
 #include "Misc/CommandLine.h"
+#include "Misc/CoreDelegates.h"
 #include "Misc/PackageName.h"
 #include "ProfilingDebugging/CpuProfilerTrace.h"
+#include "UObject/PackageReload.h"
 
-UTestPlainPropsCommandlet::UTestPlainPropsCommandlet(const FObjectInitializer& Init)
-: Super(Init)
-{}
+#if WITH_EDITOR
+#include "Engine/AssetManager.h"
+#include "Engine/Blueprint.h"
+#include "Kismet2/KismetEditorUtilities.h"
+#include "PackageTools.h"
+#endif
 
-// Slow FName based IsChildOf implementation that avoids additional dependencies
-static bool IsChildOf(const UStruct* Struct, FName Base)
+// Temporary maintained filter to ignore diffs reported by RoundtripViaLinkerBatch
+// that currently relies on IsIdentical property comparisons of loaded objects.
+// This will be replaced with a more robust and scalable approach that also captures custom bound data.
+static PlainProps::UE::FLinkerDiffFilter LinkerDiffFilter
 {
-	for (const UStruct* It : Struct->GetSuperStructIterator())
+	.BypassNativeIdenticalStructs
 	{
-		if (It->GetFName() == Base)
-		{
-			return true;
-		}
+		"FontData",							// Native FFontData::operator== with object pointers
+		"GameplayEffectVersion",			// FGameplayEffectVersion::Identical is hardcoded to return false
+		"NiagaraMeshMaterialOverride",		// Native FNiagaraMeshMaterialOverride::operator== with object pointers
+	},
+	.IgnoreStructs
+	{
+		// Runtime struct ignore diff list
+		"NiagaraCompileHash",				// FNiagaraCompileHash has just a single TArray<uint8> DataHash property
+
+		// Missing struct bindings
+		"InstancedPropertyBag",				// FInstancedPropertyBag in CoreUObject
+		"InstancedStruct",					// FInstancedStruct in CoreUObject
+		"MovieSceneFrameRange",				// FMovieSceneFrameRange in MovieScene,
+											// requires Math bindings for TRange<T> and TRangeBound<T>
+		"ShaderValueTypeHandle",			// FShaderValueTypeHandle in ComputeFramework
+		"Spline",							// FSpline in Engine,
+											// requires GeometryCore binding for TMultiSpline
+		"UniversalObjectLocator",			// FUniversalObjectLocator in Engine,
+											// requires Engine binding for FUniversalObjectLocatorFragment
+	},
+	.IgnorePropertiesForStructs
+	{
+		TPair<FName, FName>("BaseId",							"NiagaraGraphScriptUsageInfo"		),
+		TPair<FName, FName>("ChangeId",							"NiagaraNode"						),
+		TPair<FName, FName>("ChangeId",							"NiagaraGraph"						),
+		TPair<FName, FName>("LastBuiltTraversalDataChangeId",	"NiagaraGraph"						),
+		TPair<FName, FName>("LayerUsageDebugColor",				"LandscapeLayerInfoObject"			),
+		TPair<FName, FName>("UniqueId",							"NiagaraParameterDefinitionsBase"	),
+		TPair<FName, FName>("Version",							"NiagaraDataChannelVariable"		),
+	},
+	.IgnorePropertiesForBases
+	{
+		{"GraphGuid",	"EdGraph"},
+		{"NodeGuid",	"EdGraphNode"},
+	},
+	.IgnoreCastFlags = CASTCLASS_FMulticastInlineDelegateProperty | CASTCLASS_FMulticastSparseDelegateProperty,
+};
+
+static FName GenerateDroppedPackageName(FName Name, int32 NewNumber = INDEX_NONE)
+{
+	if (NewNumber == INDEX_NONE)
+	{
+		return MakeUniqueObjectName(nullptr, UPackage::StaticClass(), NAME_TrashedPackage);
+	}
+	// Transform /Game/Package_01 which is represented as "/Game/Package" and number 1 into being
+	// represented as "/Game/Package_01" and number 0, so that number 0 can be replace by the new number.
+	// This works for packages named: /Game/Package_01, /Game/Package_02, ..., /Game/Package_10
+	if (Name.GetNumber())
+	{
+		Name = FName(*Name.ToString(), 0);
 	}
-	return false;
+	return FName(Name, NewNumber + 1);
 }
 
-// Temporary maintained filter to ignore diffs reported by RoundtripViaLinkerBatch
-// that currently relies on IsIdentical property comparisons of loaded objects.
-// This will be replaced with a more robust and scalable approach that also captures custom bound data.
-static bool IgnoreLinkerDiff(const void* A, const void* B, const FProperty* Property)
+// Code inspired by UPackageTools::ReloadPackages
+static void DropPackages(TArray<UPackage*>&& Packages, int32 NewNumber)
 {
-	const static FName GameplayEffectVersion{"GameplayEffectVersion"};
-	const static FName GraphGuid{"GraphGuid"};
-	const static FName LandscapeLayerInfoObject{"LandscapeLayerInfoObject"};
-	const static FName LayerUsageDebugColor{"LayerUsageDebugColor"};
-	const static FName NiagaraParameterDefinitionsBase{"NiagaraParameterDefinitionsBase"};
-	const static FName NodeGuid{"NodeGuid"};
-	const static FName OptimusDeformer{"OptimusDeformer"};
-	const static FName Status{"Status"};
-	const static FName UniqueId{"UniqueId"};
-
-
-	const static TSet<FName> MissingStructBindings
 	{
-		"InstancedStruct",			// FInstancedStruct in CoreUObject
-		"MovieSceneFrameRange",		// FMovieSceneFrameRange in MovieScene,
-									// requires Math bindings for TRange<T> and TRangeBound<T>
-		"ShaderValueTypeHandle",	// FShaderValueTypeHandle in ComputeFramework
-		"Spline",					// FSpline in Engine,
-									// requires GeometryCore binding for TMultiSpline
-	};
-
-	const uint64 CastFlags = Property->GetCastFlags();
-	const FName PropertyName = Property->GetFName();
-	const UStruct* OwnerStruct = Property->GetOwnerStruct();
-
-	// 1. Silence diffs for known missing custom bindings for custom serialized structs.
-	// These exceptions should be removed one by one as we implement these custom bindings.
-
-	if (CastFlags & CASTCLASS_FTextProperty) // FText is tricky to bind non-intrusively
-	{
-		return true;
+		TRACE_CPUPROFILER_EVENT_SCOPE(FlushLoading);
+		FlushAsyncLoading();
+		(*GFlushStreamingFunc)();
+#if WITH_EDITOR
+		UPackageTools::FlushAsyncCompilation(Packages);
+#endif
 	}
-	else if (const FStructProperty* StructProp = CastField<FStructProperty>(Property))
+
+	TArray<UObject*> AllObjects;
+	TArray<UObject*> ObjectsInPackage;
+	TArray<FName> NewNames;
+
+	SortPackagesForReload(Packages);
+	NewNames.Reserve(Packages.Num());
+	AllObjects.Reserve(10*Packages.Num());
+	for (UPackage* Package : Packages)
 	{
-		// check missing bindings for struct properties
-		FName StructName = StructProp->Struct->GetFName();
-		if (MissingStructBindings.Contains(StructName))
-		{
-			return true;
-		}
-		else if (StructName == GameplayEffectVersion)
-		{
-			// FGameplayEffectVersion::Identical always return false to disable UE delta serialization,
-			// don't report this as a diff when the data is actually identical
-			return FMemory::Memcmp(A, B, StructProp->Struct->GetCppStructOps()->GetSize()) == 0;
-		}
+		GetObjectsWithPackage(Package, ObjectsInPackage, true, RF_Transient, EInternalObjectFlags::Garbage);
+		AllObjects.Append(ObjectsInPackage);
+		FName NewName = GenerateDroppedPackageName(Package->GetFName(), NewNumber);
+		NewNames.Add(NewName);
+		UE_LOGFMT(LogPlainPropsEngine, Log, " Dropping package {Package} with {Num} objects => {NewName}",
+			Package->GetFName(), ObjectsInPackage.Num(), NewName);
+		ObjectsInPackage.Reset();
 	}
-	else if (const FMapProperty* MapProp = CastField<FMapProperty>(Property))
+
+	FCoreDelegates::CleanupUnloadingObjects.Broadcast(AllObjects);
+
+#if WITH_EDITOR
+	for (UObject* Obj : AllObjects)
 	{
-		// check missing bindings for struct properties inside map properties
-		if (const FStructProperty* KeyProp = CastField<FStructProperty>(MapProp->GetKeyProperty()))
+		// Asset manager can hold hard references to this object and prevent GC
 		{
-			if (MissingStructBindings.Contains(KeyProp->Struct->GetFName()))
+			const FPrimaryAssetId PrimaryAssetId = UAssetManager::Get().GetPrimaryAssetIdForObject(Obj);
+			if (PrimaryAssetId.IsValid())
 			{
-				return true;
+				UAssetManager::Get().UnloadPrimaryAsset(PrimaryAssetId);
 			}
 		}
-		if (const FStructProperty* ValueProp = CastField<FStructProperty>(MapProp->GetValueProperty()))
+
+		if (UBlueprint* BP = Cast<UBlueprint>(Obj))
 		{
-			if (MissingStructBindings.Contains(ValueProp->Struct->GetFName()))
+			BP->ClearEditorReferences();
+
+			// Remove from cached dependent lists.
+			for (const TWeakObjectPtr<UBlueprint> Dependency : BP->CachedDependencies)
 			{
-				return true;
+				if (UBlueprint* ResolvedDependency = Dependency.Get())
+				{
+					ResolvedDependency->CachedDependents.Remove(BP);
+				}
 			}
-		}
-	}
-	else if (OwnerStruct->GetFName() == OptimusDeformer)
-	{
-		// Symptom from the missing FText binding.
-		// UOptimusNode::PostLoad() calls UOptimusNode_CustomComputeKernel::PostLoadNodeSpecificData()
-		// which is doing SetDisplayName()/UOptimusNodeGraph::Notify()/Deformer->MarkModified()
-		// that sets status based on a modified FText
-		return PropertyName == Status;
-	}
 
-	// 2. Ignore diffs for runtime mutated properties, e.g. properties were
-	// generation is non-deterministic, session uniqueness is enforced, generation is dependent on path name, etc
+			BP->CachedDependencies.Reset();
 
-	if (PropertyName == GraphGuid)
-	{
-		// UEdGraph::PostInitProperties() calls FGuid::NewGuid();
-		return IsChildOf(OwnerStruct, "EdGraph");
-	}
-	else if (PropertyName == NodeGuid)
-	{
-		// UEdGraphNode::CreateDeterministicGuid() calls FGuid::NewDeterministicGuid(GetPathName());
-		return IsChildOf(OwnerStruct, "EdGraphNode");
+			// Remove from cached dependency lists.
+			for (const TWeakObjectPtr<UBlueprint> Dependent : BP->CachedDependents)
+			{
+				if (UBlueprint* ResolvedDependent = Dependent.Get())
+				{
+					ResolvedDependent->CachedDependencies.Remove(BP);
+				}
+			}
+
+			BP->CachedDependents.Reset();
+		}
+		else if (UBlueprintGeneratedClass* BPGC = Cast<UBlueprintGeneratedClass>(Obj))
+		{
+			FKismetEditorUtilities::OnBlueprintGeneratedClassUnloaded.Broadcast(BPGC);
+		}
+		// else if (UWorld* World = Cast<UWorld>(Obj))
+		// {
+		// 	if (World->bIsWorldInitialized)
+		// 	{
+		// 		World->CleanupWorld();
+		// 	}
+		// }
 	}
-	else if (PropertyName == LayerUsageDebugColor)
+#endif
+
+	for (UObject* Obj : AllObjects)
 	{
-		// ULandscapeLayerInfoObject constructor is calls GenerateLayerUsageDebugColor() that depends on GetPathName()
-		return OwnerStruct->GetFName() == LandscapeLayerInfoObject;
+		Obj->SetFlags(RF_NewerVersionExists);
 	}
-	else if (OwnerStruct->GetFName() == NiagaraParameterDefinitionsBase)
+
+	for (int32 I = 0; I < Packages.Num(); ++I)
 	{
-		// UNiagaraParameterDefinitions::PostLoad() is doing FNiagaraEditorModule::EnsureReservedDefinitionUnique() to force a new unique id
-		return PropertyName == UniqueId;
+		Packages[I]->Rename(*NewNames[I].ToString(), nullptr, REN_DontCreateRedirectors | REN_NonTransactional | REN_DoNotDirty);
+		Packages[I]->SetFlags(RF_NewerVersionExists);
 	}
+}
 
-	return false;
-};
+//////////////////////////////////////////////////////////////////////////
+
+UTestPlainPropsCommandlet::UTestPlainPropsCommandlet(const FObjectInitializer& Init)
+: Super(Init)
+{}
 
 int32 UTestPlainPropsCommandlet::Main(const FString& Params)
 {
@@ -188,22 +237,23 @@ int32 UTestPlainPropsCommandlet::Main(const FString& Params)
 			{
 				FSoftObjectPath AssetPath(AssetString.Left(CommaIndex)); 
 				UObject* Object = AssetPath.TryLoad();
-				check(Object);
+				checkf(Object, TEXT("Failed to load %s (%s)"), *AssetPath.ToString(), *FString(AssetString.Left(CommaIndex)));
 				Objects.Add(Object);
 				AssetString.RightChopInline(CommaIndex + 1);
 			}
 
 			FSoftObjectPath AssetPath(AssetString);
 			UObject* Object = AssetPath.TryLoad();
-			check(Object);
+			checkf(Object, TEXT("Failed to load %s (%s)"), *AssetPath.ToString(), *FString(AssetString));
 			Objects.Add(Object);
 		}
 	}
 
 	UE_LOGFMT(LogPlainPropsEngine, Display, "FlushAsyncLoading ...");
 	{
-		TRACE_CPUPROFILER_EVENT_SCOPE(LoadAssets);
+		TRACE_CPUPROFILER_EVENT_SCOPE(InitialFlushLoading);
 		FlushAsyncLoading();
+		(*GFlushStreamingFunc)();
 	}
 	UE_LOGFMT(LogPlainPropsEngine, Display, "Starting test...");
 
@@ -257,7 +307,7 @@ int32 UTestPlainPropsCommandlet::Main(const FString& Params)
 	// Roundtrip
 	if (BatchType == EBatchType::Linker)
 	{
-		RoundtripViaLinkerBatch(Objects, Options, IgnoreLinkerDiff);
+		RoundtripViaLinkerBatch(Objects, Options, LinkerDiffFilter, DropPackages);
 	}
 	else
 	{ 
diff --git a/Engine/Plugins/Experimental/PlainPropsUObject/Source/PlainPropsUObject.Build.cs b/Engine/Plugins/Experimental/PlainPropsUObject/Source/PlainPropsUObject.Build.cs
--- a/Engine/Plugins/Experimental/PlainPropsUObject/Source/PlainPropsUObject.Build.cs
+++ b/Engine/Plugins/Experimental/PlainPropsUObject/Source/PlainPropsUObject.Build.cs
@@ -15,7 +15,7 @@ public class PlainPropsUObject : ModuleRules
 				"Core",
 				"CoreUObject",
 				"CorePreciseFP",
-				"JsonUtilities", 	// for FJsonObjectConverter::UStructToJsonObjectString
+				"JsonObjectGraph",	// for UE::JsonObjectGraph::Stringify for debugging
 				"PlainProps"
 		});
 	}
diff --git a/Engine/Plugins/Experimental/PlainPropsUObject/Source/Private/PlainPropsRoundtripTest.cpp b/Engine/Plugins/Experimental/PlainPropsUObject/Source/Private/PlainPropsRoundtripTest.cpp
--- a/Engine/Plugins/Experimental/PlainPropsUObject/Source/Private/PlainPropsRoundtripTest.cpp
+++ b/Engine/Plugins/Experimental/PlainPropsUObject/Source/Private/PlainPropsRoundtripTest.cpp
@@ -10,8 +10,10 @@
 #include "PlainPropsPrint.h"
 #include "PlainPropsSpecify.h"
 #include "PlainPropsWrite.h"
+#include "Algo/Find.h"
+#include "Async/ParallelFor.h"
 #include "HAL/FileManager.h"
-#include "JsonObjectConverter.h"
+#include "JsonObjectGraph/Stringify.h"
 #include "Logging/StructuredLog.h"
 #include "Misc/Filehelper.h"
 #include "Misc/PackageName.h"
@@ -23,6 +25,7 @@
 #include "Serialization/MemoryWriter.h"
 #include "Serialization/UnversionedPropertySerializationInternal.h"
 #include "Templates/UniquePtr.h"
+#include "Templates/UnrealTemplate.h"
 #include "UObject/LinkerLoad.h"
 #include "UObject/LinkerSave.h"
 #include "UObject/NameBatchSerialization.h"
@@ -30,6 +33,7 @@
 #include "UObject/SavePackage.h"
 #include "UObject/TextProperty.h"
 #include "UObject/UObjectIterator.h"
+#include "UObject/UObjectGlobals.h"
 #include "ProfilingDebugging/CpuProfilerTrace.h"
 
 #define PP_WITH_ENGINE_HOOKS 0
@@ -481,13 +485,27 @@ private:
 class FPersistentBatchLoader
 {
 public:
-	FPersistentBatchLoader(const FCustomBindings& Customs, FMemoryView Data)
+	FPersistentBatchLoader(const FCustomBindings& Customs, FMemoryView Data, TPair<FSensitiveName, FSensitiveName> RemapNames)
 	{
 		// Read ids
 		FByteReader It(Data);
 		verify(It.Grab<uint32>() == Magics[0]);
 		ReadStableNames(It);
 
+// todo: re-enable remap package name
+#if 0
+		// Remap original/saved names to load/instance runtime names,
+		// for now if the package name of the batch itself has been saved (e.g. for a blueprint class scope name),
+		// then replace it with its load/instance package name
+		if (RemapNames.Key != RemapNames.Value)
+		{
+			if (FSensitiveName* Name = Algo::Find(Names, FSensitiveName(RemapNames.Key)))
+			{
+				*Name = FSensitiveName(RemapNames.Value);
+			}
+		}
+#endif
+
 		// Read and mount schemas
 		verify(It.Grab<uint32>() == Magics[1]);
 		It.SkipAlignmentPadding<uint32>();
@@ -539,13 +557,6 @@ public:
 		UnmountReadSchemas(Objects[0].Schema.Batch);
 	}
 
-	FORCENOINLINE void Load(UObject* Dst)
-	{
-		FStructView In = Objects[LoadIdx];
-		LoadStruct(Dst, In.Values, In.Schema.Id, *Plans);
-		++LoadIdx;
-	}
-
 	bool HasObject(int32 ObjectIndex)
 	{
 		return ObjectMap[ObjectIndex] != ~0u;
@@ -553,17 +564,25 @@ public:
 
 	FORCENOINLINE void LoadObject(UObject* Dst, int32 ObjectIndex)
 	{
+		check(HasObject(ObjectIndex));
 		FStructView In = Objects[ObjectMap[ObjectIndex]];
 		LoadStruct(Dst, In.Values, In.Schema.Id, *Plans);
+		++NumLoaded;
+	}
+
+	static void PrintDebugInfo()
+	{
+		UE_LOGFMT(LogPlainPropsUObject, Display, "Loaded {NumLoaded} exports", NumLoaded);
 	}
 
 private:
+	static inline int32			NumLoaded = 0;
+
 	const FSchemaBatch*			TranslatedSchemas = nullptr;
 	TArray<FSensitiveName>		Names;
 	FLoadBatchPtr				Plans;
 	TArray<FStructView>			Objects;
 	TConstArrayView<uint32>		ObjectMap;
-	int32						LoadIdx = 0;
 
 	template<typename T>
 	static TConstArrayView<T> GrabNumAndArray(/* in-out */ FByteReader& It)
@@ -982,6 +1001,22 @@ public:
 		PlainProps::UE::RoundtripText(*BatchIds, Objects, {}, Format, PackageName);
 	}
 
+	void PrintPackageToYamlFile(FStringView PackageName) const
+	{
+		TUtf8StringBuilder<64> Yaml;
+		Yaml.Reserve(2 << 10);
+		PrintYamlBatch(Yaml, *BatchIds, Objects);
+		FUtf8StringView YamlView = Yaml.ToView();
+
+		// Write to file
+		const FString YamlPath = FPaths::Combine(FPaths::ProjectSavedDir(), TEXT("PlainProps"), PackageName) + TEXT(".yaml");
+		if (TUniquePtr<FArchive> Ar(IFileManager::Get().CreateFileWriter(*YamlPath)); Ar)
+		{
+			UE_LOGFMT(LogPlainPropsUObject, Log, "Writing {KB}KB yaml as {YamlPath}...", Yaml.Len() >> 10, *YamlPath);
+			Ar->Serialize((void*)YamlView.GetData(), YamlView.Len());
+		}
+	}
+
 private:
 	TArray<FStructView>			Objects;
 	TArray<FSensitiveName>		Names;
@@ -1007,15 +1042,18 @@ static FDiffDebug GPpDiff, GTpsDiff, GUpsDiff;
 
 static void DumpUObjectToJsonFile(const UObject* Object)
 {
-	FString JsonString;
-	if (FJsonObjectConverter::UStructToJsonObjectString(Object->GetClass(), Object, JsonString))
+	EJsonStringifyFlags StringifyFlags = EJsonStringifyFlags::DisableDeltaEncoding | EJsonStringifyFlags::DisableNativeSerializers;
+	if (FUtf8String Json(::UE::JsonObjectGraph::Stringify({Object}, {StringifyFlags})); Json.Len())
 	{
 		TStringBuilder<256> JsonPath;
 		JsonPath << FPaths::Combine(FPaths::ProjectSavedDir(), TEXT("PlainProps"), Object->GetPackage()->GetName());
 		JsonPath << "." << Object->GetName() + ".json";
 
-		UE_LOGFMT(LogPlainPropsUObject, Display, "Writing {JsonPath}...", JsonPath);
-		FFileHelper::SaveStringToFile(JsonString, JsonPath.ToString(), FFileHelper::EEncodingOptions::ForceUTF8WithoutBOM);
+		UE_LOGFMT(LogPlainPropsUObject, Log, "Writing {KB}KB json as {JsonPath}...", Json.Len() >> 10, *JsonPath);
+		if (TUniquePtr<FArchive> Ar(IFileManager::Get().CreateFileWriter(JsonPath.ToString())); Ar)
+		{
+			Ar->Serialize((void*)Json.GetCharArray().GetData(), Json.GetCharArray().Num() - 1);
+		}
 	}
 }
 
@@ -1089,8 +1127,7 @@ struct FLinkerPropertyBatch
 
 	void Load(FText& Out, FStructLoadView In) const
 	{
-		int32 Dummy = LoadSole<int32>(In);
-		Out = {};//TODO: Texts[LoadSole<int32>(In)];
+		Out = Texts[LoadSole<int32>(In)];
 	}
 
 	static FSoleMemberSpec Spec(FObjectHandle*, FDeclId Id)
@@ -1123,7 +1160,7 @@ struct FLinkerPropertyBatch
 			{
 				UObject* Object = LinkerLoad->ImportMap[PackageIdx.ToImport()].XObject;
 #if PP_WITH_ENGINE_HOOKS
-				//Object = Object ? Object : LinkerLoad->ResolveResource(PackageIdx);
+				Object = Object ? Object : LinkerLoad->ResolveResource(PackageIdx);
 #endif
 				Out = ::UE::CoreUObject::Private::MakeObjectHandle(Object);
 			}
@@ -1231,6 +1268,9 @@ struct FLinkerPropertyBatch
 		LoadSoleStruct(&Out.GetUniqueID(), In);
 	}
 
+	TArray<FText>&	GrabTexts()							{ return Texts; }
+	void			GiveTexts(TArray<FText>&& InTexts)	{ Texts = MoveTemp(InTexts); }
+
 private:
 	union
 	{
@@ -1266,12 +1306,14 @@ struct FPlainPropsLinkerSandbox
 		FPackagePath PackagePath;
 		FSharedBuffer PackageData;
 		TArray64<uint8> PlainPropsData;
+		TArray<FText> Texts;
 	};
 
-	struct FLinkerLoadInstance
+	struct FMountedLinker
 	{
-		FLinkerLoadInstance(const FSavedPackage& Saved, uint32 Number)
-			: PackageName(Saved.PackageName.GetNumber() ? FName(*Saved.PackageName.ToString(), Number) : FName(Saved.PackageName, Number))
+		FMountedLinker(const FSavedPackage& Saved, FName InPackageName)
+			: PackageName(InPackageName)
+			, PackagePath(Saved.PackagePath)
 			, Package(NewObject<UPackage>(/*Outer*/nullptr, PackageName))
 			, LinkerReader(new FLargeMemoryReader(static_cast<const uint8*>(Saved.PackageData.GetData()), Saved.PackageData.GetSize(), ELargeMemoryReaderFlags::Persistent, PackageName))
 			, LinkerLoad(FLinkerLoad::CreateLinker(FUObjectThreadContext::Get().GetSerializeContext(), Package, Saved.PackagePath, LoadFlags, LinkerReader))
@@ -1279,59 +1321,57 @@ struct FPlainPropsLinkerSandbox
 			check(LinkerLoad->GetLoader_Unsafe() == LinkerReader);
 		}
 
-		~FLinkerLoadInstance()
+		~FMountedLinker()
 		{
 			// This will first destroy and free the LinkerReader then the LinkerLoad
 			ResetLoaders(Package);
 		}
 
 		FName PackageName;
+		FPackagePath PackagePath;
 		UPackage* Package;
 		FLargeMemoryReader* LinkerReader;
 		FLinkerLoad* LinkerLoad;
 	};
 
-	struct FMountedPackage
+	struct FMountedLinkerBatch
 	{
-		FMountedPackage(const FSavedPackage& Saved)
-			: PackagePath(Saved.PackagePath)
-			, Base(Saved, BaseNumber)
-			, PP(Saved, PPNumber)
-			, Properties(*PP.LinkerLoad)
+		FMountedLinkerBatch(const FSavedPackage& Saved, FName InPackageName)
+			: MountedLinker(Saved, InPackageName)
+			, Properties(*MountedLinker.LinkerLoad)
 			, Customs(Properties, GUE.Customs)
 			, BatchView(MakeMemoryView(Saved.PlainPropsData))
-			, BatchLoader(Customs.Overlay, BatchView)
+			, RemapPackageName(FSensitiveName(Saved.PackageName), FSensitiveName(InPackageName))
+			, BatchLoader(Customs.Overlay, BatchView, RemapPackageName)
 		{
+			Properties.GiveTexts(CopyTemp(Saved.Texts));
 		}
 
-		FPackagePath PackagePath;
-		FLinkerLoadInstance Base;
-		FLinkerLoadInstance PP;
+		FMountedLinker MountedLinker;
 		FLinkerPropertyBatch Properties;
 		FLinkerPropertyBindings Customs;
 		FMemoryView BatchView;
+		TPair<FSensitiveName, FSensitiveName> RemapPackageName;
 		FPersistentBatchLoader BatchLoader;
 	};
 
 	static TMap<FName, FSavedPackage> SavedPackages;
-	static TMap<FName, TUniquePtr<FMountedPackage>> MountedPackages;
+	static TMap<FName, TUniquePtr<FMountedLinker>> MountedBaselinePackages;
+	static TMap<FName, TUniquePtr<FMountedLinkerBatch>> MountedPlainpropsPackages;
+	static TArray<const UClass*> PendingClasses;
 
-	// Transforms /Game/Package_01 which is represented as "/Game/Package" and number 1 into
-	// being represented as "/Game/Package_01" and number 0, so that number 0 can be replace by the two magic numbers 8848 and 8849.
-	// This works for packages named:/Game/Package_01, /Game/Package_02, ..., /Game/Package_10
-	static FName FoldNameNumber(FName Name)
+#if WITH_EDITOR
+	static void OnBlueprintCDOCompiled(UObject* InObject, const FObjectPostCDOCompiledContext&)
 	{
-		if (Name.GetNumber())
-		{
-			Name = FName(*Name.ToString(), 0);
-		}
-		return Name;
+		const UClass* Class = InObject->GetClass();
+		PendingClasses.AddUnique(Class);
 	}
+#endif
 
 	static void SavePlainProps(const FLinkerSave& LinkerSave)
 	{
 		TRACE_CPUPROFILER_EVENT_SCOPE(SavePlainProps);
-		FName PackageName = FoldNameNumber(LinkerSave.LinkerRoot->GetFName());
+		FName PackageName = LinkerSave.LinkerRoot->GetFName();
 
 		UE_LOGFMT(LogPlainPropsUObject, Log, " Saving PlainProps for {PackageName}...", PackageName);
 		FLinkerPropertyBatch Properties(LinkerSave);
@@ -1356,18 +1396,19 @@ struct FPlainPropsLinkerSandbox
 		SavedPackage.PackageName = PackageName;
 		SavedPackage.PackagePath = LinkerSave.LinkerRoot->GetLoadedPath();
 		SavedPackage.PlainPropsData = MoveTemp(StableData);
+		SavedPackage.Texts = Properties.GrabTexts();
 	}
 
 	static void SavePackageToMemory(const FLinkerSave& LinkerSave, uint8* Data, int64 Size)
 	{
-		FName PackageName = FoldNameNumber(LinkerSave.LinkerRoot->GetFName());
+		FName PackageName = LinkerSave.LinkerRoot->GetFName();
 		FSavedPackage& SavedPackage = SavedPackages.FindChecked(PackageName);
 		SavedPackage.PackageData = FSharedBuffer::TakeOwnership(Data, Size, FMemory::Free);
 	}
 
 	static bool OpenPackageReader(const FPackagePath& PackagePath, FArchive*& Ar)
 	{
-		FName PackageName = FoldNameNumber(PackagePath.GetPackageFName());
+		FName PackageName = PackagePath.GetPackageFName();
 		if (FSavedPackage* SavedData = SavedPackages.Find(PackageName))
 		{
 			Ar = new FLargeMemoryReader(static_cast<const uint8*>(SavedData->PackageData.GetData()), SavedData->PackageData.GetSize(), ELargeMemoryReaderFlags::Persistent, PackageName);
@@ -1376,44 +1417,68 @@ struct FPlainPropsLinkerSandbox
 		return false;
 	}
 
-	static void MountPackagesFromMemory(FName PackageName)
+	[[nodiscard]] static UPackage* MountBaselinePackage(FName PackageName, bool bClone)
+	{
+		if (FSavedPackage* SavedData = SavedPackages.Find(PackageName))
+		{
+			FName NewPackageName = PackageName;
+			if (bClone)
+			{
+				NewPackageName = *WriteToString<256>("/CloneBaseline", PackageName);
+			}
+			UE_LOGFMT(LogPlainPropsUObject, Log, " Mounting baseline for {Package}...", NewPackageName);
+			TUniquePtr<FMountedLinker> MountedInstance(new FMountedLinker(*SavedData, NewPackageName));
+			UPackage* Package = MountedInstance->Package;
+			MountedBaselinePackages.Emplace(NewPackageName, MoveTemp(MountedInstance));
+			return Package;
+		}
+		return nullptr;
+	}
+
+	[[nodiscard]] static UPackage* MountPlainPropsPackage(FName PackageName, bool bClone)
 	{
-		PackageName = FoldNameNumber(PackageName);
 		if (FSavedPackage* SavedData = SavedPackages.Find(PackageName))
 		{
-			FMountedPackage* LinkerBatch = new FMountedPackage(*SavedData);
-			MountedPackages.Emplace(PackageName, LinkerBatch);
+			FName NewPackageName = PackageName;
+			if (bClone)
+			{
+				NewPackageName = *WriteToString<256>("/ClonePlainProps", PackageName);
+			}
+			UE_LOGFMT(LogPlainPropsUObject, Log, " Mounting plainprops for {Package}...", NewPackageName);
+			TUniquePtr<FMountedLinkerBatch> MountedInstance(new FMountedLinkerBatch(*SavedData, NewPackageName));
+			UPackage* Package = MountedInstance->MountedLinker.Package;
+			MountedPlainpropsPackages.Add(NewPackageName, MoveTemp(MountedInstance));
+			return Package;
 		}
+		return nullptr;
 	}
 
 	static void Shutdown()
 	{
-		MountedPackages.Empty();
+		MountedPlainpropsPackages.Empty();
 		SavedPackages.Empty();
 	}
 
-	static void LoadBaselinePackageFromMemory(FName PackageName)
+	static UPackage* LoadBaselinePackage(UPackage* Package)
 	{
-		PackageName = FoldNameNumber(PackageName);
-		if (TUniquePtr<FMountedPackage>* LinkerBatchPtr = MountedPackages.Find(PackageName))
+		if (TUniquePtr<FMountedLinker>* MountedLinkerPtr = MountedBaselinePackages.Find(Package->GetFName()))
 		{
-			FMountedPackage& LinkerBatch = **LinkerBatchPtr;
+			FMountedLinker& MountedLinker = **MountedLinkerPtr;
 
-			UE_LOGFMT(LogPlainPropsUObject, Log, " Loading baseline for {Package}...", PackageName);
-			LoadPackage(LinkerBatch.Base.Package, LinkerBatch.PackagePath, LoadFlags);
+			UE_LOGFMT(LogPlainPropsUObject, Log, " Loading baseline for {Package}...", MountedLinker.PackageName);
+			return LoadPackage(MountedLinker.Package, MountedLinker.PackagePath, LoadFlags);
 		}
+		return nullptr;
 	}
 
-	static UPackage* LoadPlainPropsPackageFromMemory(FName PackageName)
+	static UPackage* LoadPlainPropsPackage(UPackage* Package)
 	{
-		PackageName = FoldNameNumber(PackageName);
-		if (TUniquePtr<FMountedPackage>* LinkerBatchPtr = MountedPackages.Find(PackageName))
+		if (TUniquePtr<FMountedLinkerBatch>* LinkerBatchPtr = MountedPlainpropsPackages.Find(Package->GetFName()))
 		{
-			FMountedPackage& LinkerBatch = **LinkerBatchPtr;
-
-			UE_LOGFMT(LogPlainPropsUObject, Log, " Loading PlainProps for {Package}...", PackageName);
+			FMountedLinkerBatch& LinkerBatch = **LinkerBatchPtr;
 
-			return LoadPackage(LinkerBatch.PP.Package, LinkerBatch.PackagePath, LoadFlags);
+			UE_LOGFMT(LogPlainPropsUObject, Log, " Loading PlainProps for {Package}...", LinkerBatch.MountedLinker.PackageName);
+			return LoadPackage(LinkerBatch.MountedLinker.Package, LinkerBatch.MountedLinker.PackagePath, LoadFlags);
 		}
 
 		return nullptr;
@@ -1421,7 +1486,6 @@ struct FPlainPropsLinkerSandbox
 
 	static void RoundtripToTextFromMemory(FName PackageName)
 	{
-		PackageName = FoldNameNumber(PackageName);
 		if (FSavedPackage* SavedData = SavedPackages.Find(PackageName))
 		{
 			FBatchTextRoundtripper StableBatch(MakeMemoryView(SavedData->PlainPropsData), ESchemaFormat::StableNames);
@@ -1429,19 +1493,13 @@ struct FPlainPropsLinkerSandbox
 		}
 	}
 
-	static FDiffObjectContext DiffPackagesLoadedFromMemory(FName PackageName, FDiffObjectContext::FIgnoreFilter IgnoreDiffFilter)
+	static void PrintPackageToYamlFile(FName PackageName)
 	{
-		PackageName = FoldNameNumber(PackageName);
-		if (TUniquePtr<FMountedPackage>* LinkerBatchPtr = MountedPackages.Find(PackageName))
+		if (FSavedPackage* SavedData = SavedPackages.Find(PackageName))
 		{
-			FMountedPackage& LinkerBatch = **LinkerBatchPtr;
-
-			FDiffObjectContext Ctx = { LinkerBatch.Base.Package, LinkerBatch.PP.Package };
-			Ctx.IgnoreDiff = IgnoreDiffFilter;
-			DiffObjects(Ctx);
-			return MoveTemp(Ctx);
+			FBatchTextRoundtripper StableBatch(MakeMemoryView(SavedData->PlainPropsData), ESchemaFormat::StableNames);
+			StableBatch.PrintPackageToYamlFile(PackageName.ToString());
 		}
-		return {};
 	}
 
 	static bool LoadObject(FStructuredArchive::FSlot& Slot, uint8* Data)
@@ -1454,18 +1512,12 @@ struct FPlainPropsLinkerSandbox
 
 		UObject* Object = (UObject*)Data;
 		FName PackageName = Object->GetPackage()->GetFName();
-		if (PackageName.GetNumber() != PPNumber)
-		{
-			return false;
-		}
-		PackageName.SetNumber(0);
-
-		TUniquePtr<FMountedPackage>* LinkerBatchPtr = MountedPackages.Find(PackageName);
+		TUniquePtr<FMountedLinkerBatch>* LinkerBatchPtr = MountedPlainpropsPackages.Find(PackageName);
 		check(LinkerBatchPtr);
-		FMountedPackage& LinkerBatch = **LinkerBatchPtr;
+		FMountedLinkerBatch& LinkerBatch = **LinkerBatchPtr;
 
 		FPersistentBatchLoader& BatchLoader = LinkerBatch.BatchLoader;
-		FLinkerLoad* LinkerLoad = LinkerBatch.PP.LinkerLoad;
+		FLinkerLoad* LinkerLoad = LinkerBatch.MountedLinker.LinkerLoad;
 		check(LinkerLoad == Object->GetLinker());
 
 		FUObjectSerializeContext* SerializeContext = FUObjectThreadContext::Get().GetSerializeContext();
@@ -1475,18 +1527,58 @@ struct FPlainPropsLinkerSandbox
 		const FObjectExport& Export = LinkerLoad->ExportMap[ExportIndex];
 		check(Export.Object == Object);
 
+		if (PendingClasses.Num())
+		{
+			// Re-bind dynamic classes whenever their CDO has been recompiled
+			for (const UClass* Class : PendingClasses)
+			{
+				FBindId Id = GUE.Names.IndexBindId(PlainProps::UE::IndexStruct(Class));
+				if (GUE.Schemas.FindStruct(Id))
+				{
+					UE_LOGFMT(LogPlainPropsUObject, Log, " Re-binding compiled blueprint class {BindId}", GUE.Debug.Print(Id));
+					GUE.Schemas.DropStruct(Id);
+				}
+				else
+				{
+					UE_LOGFMT(LogPlainPropsUObject, Log, " Binding compiled blueprint class {BindId}", GUE.Debug.Print(Id));
+				}
+				UE::BindStruct(Id, Class);
+			}
+			PendingClasses.Reset();
+		}
+		// Bind blueprint generated classes at the time we load its CDO.
+		// This is first time the struct binding is needed,
+		// and we also know for sure that the UClass has already been linked.
+		if (Object->HasAllFlags(RF_ClassDefaultObject))
+		{
+			UClass* Class = Object->GetClass();
+			if (!Class->HasAllClassFlags(CLASS_Native))
+			{
+				FBindId Id = GUE.Names.IndexBindId(PlainProps::UE::IndexStruct(Class));
+				UE_LOGFMT(LogPlainPropsUObject, Log, " Binding loaded blueprint class {BindId} for loading CDO {CDO}",
+					GUE.Debug.Print(Id), *Object->GetPathName());
+				check(!GUE.Schemas.FindStruct(Id))
+				UE::BindStruct(Id, Class);
+			}
+		}
 		{
 			TRACE_CPUPROFILER_EVENT_SCOPE(LoadPlainProps);
 			BatchLoader.LoadObject(Object, ExportIndex);
 		}
+// todo: re-enable bind dynamic class
+#if 0
+		// Bind dynamic classes other than blueprints right after we have loaded them
 		if (const UClass* Class = Cast<UClass>(Object))
 		{
-			FBindId Id = GUE.Names.IndexBindId(PlainProps::UE::IndexStruct(Class));
-			UE_LOGFMT(LogPlainPropsUObject, Display, " Binding loaded dynamic class {BindId}", GUE.Debug.Print(Id));
-			check(!GUE.Schemas.FindStruct(Id))
-			UE::BindStruct(Id, Class);
+			if (!Class->HasAnyClassFlags(CLASS_CompiledFromBlueprint))
+			{
+				FBindId Id = GUE.Names.IndexBindId(PlainProps::UE::IndexStruct(Class));
+				UE_LOGFMT(LogPlainPropsUObject, Display, " Binding loaded dynamic class {BindId}", GUE.Debug.Print(Id));
+				check(!GUE.Schemas.FindStruct(Id))
+				UE::BindStruct(Id, Class);
+			}
 		}
-
+#endif
 #if WITH_EDITOR
 		const int32 Pos = Ar.Tell();
 		check(Pos == SerializeContext->SerializedObjectScriptStartOffset);
@@ -1497,12 +1589,15 @@ struct FPlainPropsLinkerSandbox
 		// todo: SerializedObjectScriptStartOffset does not exist in client builds
 		unimplemented();
 #endif
+
 		return true;
 	}
 };
 
 TMap<FName, FPlainPropsLinkerSandbox::FSavedPackage> FPlainPropsLinkerSandbox::SavedPackages;
-TMap<FName, TUniquePtr<FPlainPropsLinkerSandbox::FMountedPackage>> FPlainPropsLinkerSandbox::MountedPackages;
+TMap<FName, TUniquePtr<FPlainPropsLinkerSandbox::FMountedLinkerBatch>> FPlainPropsLinkerSandbox::MountedPlainpropsPackages;
+TMap<FName, TUniquePtr<FPlainPropsLinkerSandbox::FMountedLinker>> FPlainPropsLinkerSandbox::MountedBaselinePackages;
+TArray<const UClass*> FPlainPropsLinkerSandbox::PendingClasses;
 
 #endif
 
@@ -1771,7 +1866,7 @@ int32 RoundtripViaPlainBatch(TConstArrayView<UObject*> Objects, ERoundtrip Optio
 #if PP_WITH_ENGINE_HOOKS
 struct FScopeSaveTestPackage
 {
-	FScopeSaveTestPackage()
+	FORCENOINLINE FScopeSaveTestPackage()
 	{
 		FSavePackageSettings& SavePackageSettings = FSavePackageSettings::GetDefaultSettings();
 		ExternalWriteExportsFunc = SavePackageSettings.ExternalWriteExportsFunc;
@@ -1794,22 +1889,30 @@ struct FScopeSaveTestPackage
 
 struct FScopeLoadTestPackage
 {
-	FScopeLoadTestPackage()
+	FScopeLoadTestPackage(bool bInPlainProps)
+		: SerializePropertiesFunc(GSerializePropertiesOverride)
+		, LoadFromPackageFileFunc(GLoadFromPackageFileOverride)
+		, bPlainProps(bInPlainProps)
 	{
-		SerializePropertiesFunc = GSerializePropertiesOverride;
-		LoadFromPackageFileFunc = GLoadFromPackageFileOverride;
-		GSerializePropertiesOverride = &FPlainPropsLinkerSandbox::LoadObject;
+		if (bPlainProps)
+		{
+			GSerializePropertiesOverride = &FPlainPropsLinkerSandbox::LoadObject;
+		}
 		GLoadFromPackageFileOverride = &FPlainPropsLinkerSandbox::OpenPackageReader;
 	}
 
 	~FScopeLoadTestPackage()
 	{
-		GSerializePropertiesOverride = SerializePropertiesFunc;
+		if (bPlainProps)
+		{
+			GSerializePropertiesOverride = SerializePropertiesFunc;
+		}
 		GLoadFromPackageFileOverride = LoadFromPackageFileFunc;
 	}
 
 	bool(*SerializePropertiesFunc)(class FStructuredArchiveSlot&, uint8*);
 	bool (*LoadFromPackageFileFunc)(const FPackagePath&, FArchive*&) = nullptr;
+	bool bPlainProps;
 };
 #else
 struct FScopeEmptyBase 
@@ -1818,56 +1921,128 @@ struct FScopeEmptyBase
 	~FScopeEmptyBase() {}
 };
 struct FScopeSaveTestPackage : public FScopeEmptyBase {};
-struct FScopeLoadTestPackage : public FScopeEmptyBase {};
+struct FScopeLoadTestPackage : public FScopeEmptyBase
+{
+	FScopeLoadTestPackage(bool) {};
+};
 #endif // PP_WITH_ENGINE_HOOKS
 
-int32 RoundtripViaLinkerBatch(TConstArrayView<UObject*> Objects, ERoundtrip Options, FLinkerDiffIgnoreFilter DiffIgnoreFilter)
+// Transforms /Game/Package_01 which is represented as "/Game/Package" and number 1 into being
+// represented as "/Game/Package_01" and number 0, so that number 0 can be replace by a new number.
+// This works for packages named: /Game/Package_01, /Game/Package_02, ..., /Game/Package_10
+static FName GenerateDroppedPackageName(FName Name, int32 NewNumber = 0)
 {
-	TRACE_CPUPROFILER_EVENT_SCOPE(RoundtripViaLinkerBatch);
-
-	TArray<UPackage*> Packages;
-	TArray<UObject*> AllPackages;
-
-	// Collect all packages
-	if (Objects.Num())
+	if (NewNumber == 0)
 	{
-		for (UObject* Object : Objects)
-		{
-			AllPackages.Add(Object->GetPackage());
-		}
+		return MakeUniqueObjectName(nullptr, UPackage::StaticClass(), NAME_TrashedPackage);
 	}
-	else
+	if (Name.GetNumber())
 	{
-		GetObjectsOfClass(UPackage::StaticClass(), AllPackages);
+		Name = FName(*Name.ToString(), 0);
 	}
+	return FName(Name, NewNumber);
+}
 
-	// Filter out in-memory packages
-	for (UObject* PackageObj : AllPackages)
+static void DropPackages(TArray<UPackage*>&& Packages, FDropPackagesFunc DropPackagesFunc, int32 NewNumber = INDEX_NONE)
+{
+// todo: re-enable drop packages
+#if 0
+	for (UPackage* Package : Packages)
 	{
-		if (!PackageObj->HasAnyFlags(RF_WasLoaded))
+		ForEachObjectWithPackage(Package, [](UObject* Object)
 		{
-			continue;
+			if (const UClass* Class = Cast<UClass>(Object))
+			{
+				FBindId Id = GUE.Names.IndexBindId(PlainProps::UE::IndexStruct(Class));
+				if (GUE.Schemas.FindStruct(Id))
+				{
+					UE_LOGFMT(LogPlainPropsUObject, Log, " Dropping loaded dynamic class {BindId}", GUE.Debug.Print(Id));
+					GUE.Schemas.DropStruct(Id);
+				}
+			}
+			return true;
+		});
+	}
+
+	DropPackagesFunc(MoveTemp(Packages), NewNumber);
+#endif
+}
+
+int32 RoundtripViaLinkerBatch(
+	TConstArrayView<UObject*> Objects,
+	ERoundtrip Options,
+	const FLinkerDiffFilter& DiffFilter,
+	FDropPackagesFunc DropPackagesFunc)
+{
+#if PP_WITH_ENGINE_HOOKS
+	TRACE_CPUPROFILER_EVENT_SCOPE(RoundtripViaLinkerBatch);
+
+	TArray<FName> PackageNames;
+	TArray<UPackage*> OriginalPackages;
+	TArray<UPackage*> BaselinePackages;
+	TArray<UPackage*> PlainPropsPackages;
+
+	{
+		TArray<UObject*> AllPackages;
+		// Collect all packages
+		if (!Objects.Num())
+		{
+			GetObjectsOfClass(UPackage::StaticClass(), AllPackages);
+		}
+		// Collect specified packages
+		if (Objects.Num())
+		{
+			for (UObject* Object : Objects)
+			{
+				AllPackages.AddUnique(Object->GetPackage());
+			}
 		}
-		UPackage* Package = Cast<UPackage>(PackageObj);
-		if (Package->HasAnyPackageFlags(PKG_CompiledIn))
+
+		// Filter out in-memory packages
+		for (UObject* PackageObj : AllPackages)
 		{
-			continue;
+			if (!PackageObj->HasAnyFlags(RF_WasLoaded))
+			{
+				continue;
+			}
+			UPackage* Package = Cast<UPackage>(PackageObj);
+			if (Package->HasAnyPackageFlags(PKG_CompiledIn))
+			{
+				continue;
+			}
+			// todo: re-enable blueprints
+			bool bExclude = false;
+			ForEachObjectWithPackage(Package, [&bExclude](UObject* Object)
+			{
+				UClass* Class = Object->GetClass();
+				if (Class->HasAnyClassFlags(CLASS_CompiledFromBlueprint))
+				{
+					bExclude = true;
+					return false;
+				}
+				return true;
+			});
+			if (bExclude)
+			{
+				continue;
+			}
+			OriginalPackages.Add(Package);
+			PackageNames.Add(Package->GetFName());
 		}
-		Packages.Add(Package);
 	}
+	const int32 NumPackages = OriginalPackages.Num();
 
-#if PP_WITH_ENGINE_HOOKS
-	FPlainPropsLinkerSandbox::SavedPackages.Reserve(Packages.Num());
-	FPlainPropsLinkerSandbox::MountedPackages.Reserve(Packages.Num());
-#endif
+	FPlainPropsLinkerSandbox::SavedPackages.Reserve(NumPackages);
+	FPlainPropsLinkerSandbox::MountedBaselinePackages.Reserve(NumPackages);
+	FPlainPropsLinkerSandbox::MountedPlainpropsPackages.Reserve(NumPackages);
 
-	UE_LOGFMT(LogPlainPropsUObject, Display, "Saving {SaveNum} packages...", Packages.Num());
+	UE_LOGFMT(LogPlainPropsUObject, Display, "Saving {SaveNum} packages...", NumPackages);
 
 	{
 		FScopeSaveTestPackage SaveScope;
 
 		TRACE_CPUPROFILER_EVENT_SCOPE(SavePackages);
-		for (UPackage* Package : Packages)
+		for (UPackage* Package : OriginalPackages)
 		{
 			UObject* Asset = Package->FindAssetInPackage();
 			FSavePackageArgs SaveArgs;
@@ -1883,71 +2058,163 @@ int32 RoundtripViaLinkerBatch(TConstArrayView<UObject*> Objects, ERoundtrip Opti
 		}
 	}
 
-#if PP_WITH_ENGINE_HOOKS
 	// Roundtrip to text, also dumping all packages as yaml to disk
 	if (EnumHasAnyFlags(Options, ERoundtrip::TextMemory | ERoundtrip::TextStable))
 	{
 		TRACE_CPUPROFILER_EVENT_SCOPE(RoundtripToText);
-		UE_LOGFMT(LogPlainPropsUObject, Display, "Roundtripping {SaveNum} packages to text...", Packages.Num());
-		for (UPackage* Package : Packages)
+		UE_LOGFMT(LogPlainPropsUObject, Display, "Roundtripping {SaveNum} packages to text...", NumPackages);
+		for (FName PackageName : PackageNames)
 		{
-			FPlainPropsLinkerSandbox::RoundtripToTextFromMemory(Package->GetFName());
+			FPlainPropsLinkerSandbox::RoundtripToTextFromMemory(PackageName);
 		}
 	}
 
+#if WITH_EDITOR
+	FCoreUObjectDelegates::OnObjectPostCDOCompiled.AddStatic(FPlainPropsLinkerSandbox::OnBlueprintCDOCompiled);
+#endif
+
 	// Roundtrip plain props and diff against baseline (TPS or UPS)
 	if (EnumHasAnyFlags(Options, ERoundtrip::PP))
 	{
-		FScopeLoadTestPackage LoadScope;
+		TBitArray<> InstancedSet;
+		InstancedSet.Reserve(NumPackages);
+		{
+			TRACE_CPUPROFILER_EVENT_SCOPE(DropOriginalPackages);
+			UE_LOGFMT(LogPlainPropsUObject, Display, "Dropping {Num} original packages...", NumPackages);
+			DropPackages(CopyTemp(OriginalPackages), DropPackagesFunc);
+			// Some packages may have been reloaded as a result of the drop package callbacks,
+			// mark that these should be loaded as instanced packages, i.e. with new unique in-memory package names
+			for (FName PackageName : PackageNames)
+			{
+				InstancedSet.Add(FindObjectFast<UPackage>(/*Outer*/nullptr, PackageName) != nullptr);
+			}
+		}
 
 		{
-			UE_LOGFMT(LogPlainPropsUObject, Display, "Mounting and loading {SaveNum} packages twice...", Packages.Num());
-			for (UPackage* Package : Packages)
+			TRACE_CPUPROFILER_EVENT_SCOPE(MountPlainPropsPackages);
+			UE_LOGFMT(LogPlainPropsUObject, Display, "Mounting {Num} plain props packages...", NumPackages);
+			for (int32 I = 0; I < NumPackages; ++I)
 			{
-				FPlainPropsLinkerSandbox::MountPackagesFromMemory(Package->GetFName());
+				PlainPropsPackages.Add(FPlainPropsLinkerSandbox::MountPlainPropsPackage(PackageNames[I], InstancedSet[I]));
 			}
 		}
+
 		{
 			TRACE_CPUPROFILER_EVENT_SCOPE(LoadPlainPropsPackages);
+			UE_LOGFMT(LogPlainPropsUObject, Display, "Loading {Num} plain props packages...", NumPackages);
+			FScopeLoadTestPackage LoadScope(/*bPlainProps*/true);
 			double Start = FPlatformTime::Seconds();
-			for (UPackage* Package : Packages)
+			for (UPackage* Package : PlainPropsPackages)
 			{
-				FPlainPropsLinkerSandbox::LoadPlainPropsPackageFromMemory(Package->GetFName());
+				FPlainPropsLinkerSandbox::LoadPlainPropsPackage(Package);
 			}
 			UE_LOG(LogPlainPropsUObject, Display, TEXT("Loaded %d plainprops packages in %.2f s"),
-				Packages.Num(), FPlatformTime::Seconds() - Start);
+				NumPackages, FPlatformTime::Seconds() - Start);
+		}
+
+		{
+			TRACE_CPUPROFILER_EVENT_SCOPE(DropPlainPropsPackages);
+			UE_LOGFMT(LogPlainPropsUObject, Display, "Dropping {Num} plain props packages...", NumPackages);
+			DropPackages(CopyTemp(PlainPropsPackages), DropPackagesFunc, 8849);
+		}
+
+		{
+			TRACE_CPUPROFILER_EVENT_SCOPE(MountPlainPropsPackages);
+			UE_LOGFMT(LogPlainPropsUObject, Display, "Mounting {Num} baseline packages...", NumPackages);
+			for (int32 I = 0; I < NumPackages; ++I)
+			{
+				BaselinePackages.Add(FPlainPropsLinkerSandbox::MountBaselinePackage(PackageNames[I], InstancedSet[I]));
+			}
 		}
 		{
-			TRACE_CPUPROFILER_EVENT_SCOPE(LoadBaselinePackags);
+			TRACE_CPUPROFILER_EVENT_SCOPE(LoadBaselinePackages);
+			UE_LOGFMT(LogPlainPropsUObject, Display, "Loading {Num} baseline packages...", NumPackages);
+			FScopeLoadTestPackage LoadScope(/*bPlainProps*/false);
 			double Start = FPlatformTime::Seconds();
-			for (UPackage* Package : Packages)
+			for (UPackage* Package : BaselinePackages)
 			{
-				FPlainPropsLinkerSandbox::LoadBaselinePackageFromMemory(Package->GetFName());
+				FPlainPropsLinkerSandbox::LoadBaselinePackage(Package);
 			}
 			UE_LOG(LogPlainPropsUObject, Display, TEXT("Loaded %d baseline packages in %.2f s"),
-				Packages.Num(), FPlatformTime::Seconds() - Start);
+				NumPackages, FPlatformTime::Seconds() - Start);
 		}
+
+		{
+			TRACE_CPUPROFILER_EVENT_SCOPE(DropBaselinePackages);
+			UE_LOGFMT(LogPlainPropsUObject, Display, "Dropping {Num} baseline packages...", NumPackages);
+			DropPackages(CopyTemp(BaselinePackages), DropPackagesFunc, 8848);
+		}
+
 		{
 			TRACE_CPUPROFILER_EVENT_SCOPE(DiffPackages);
-			UE_LOGFMT(LogPlainPropsUObject, Display, "Diffing {SaveNum} packages...\n", Packages.Num());
-			TUtf8StringBuilder<4096> Diffs;
-			uint32 NumDiffs = 0;
-			for (UPackage* Package : Packages)
+			UE_LOGFMT(LogPlainPropsUObject, Display, "Diffing {Num} packages...\n", NumPackages);
+
+			EParallelForFlags Flags = EParallelForFlags::Unbalanced;
+			Flags |= FPlatformMisc::IsDebuggerPresent() ? EParallelForFlags::ForceSingleThread : EParallelForFlags::None;
+
+			FDiffObjectFilter Filter = 
 			{
-				FDiffObjectContext Ctx = FPlainPropsLinkerSandbox::DiffPackagesLoadedFromMemory(Package->GetFName(), DiffIgnoreFilter);
+				DiffFilter.BypassNativeIdenticalStructs,
+				DiffFilter.IgnoreStructs,
+				DiffFilter.IgnorePropertiesForStructs,
+				DiffFilter.IgnorePropertiesForBases,
+				DiffFilter.IgnoreCastFlags,
+			};
+
+			std::atomic<uint32> NumDiffs = 0;
+			std::atomic<uint32> NumIgnored = 0;
+			ParallelFor(TEXT("DiffPackage"), NumPackages, /*MinBatchSize*/1,
+				[&NumDiffs, &NumIgnored, &PackageNames, &BaselinePackages, &PlainPropsPackages, &Filter](int32 I)
+			{
+				FDiffObjectContext Ctx = { BaselinePackages[I], PlainPropsPackages[I], Filter };
+				{
+					TRACE_CPUPROFILER_EVENT_SCOPE(DiffObjects);
+					DiffObjects(Ctx);
+				}
 				if (Ctx.Diffs.Num())
 				{
-					PrintDiff(Diffs, Ctx);
-					Diffs << " in " << Package->GetFName() << '\n';
+					TUtf8StringBuilder<2048> Diff;
+					PrintDiff(Diff, Ctx);
+					UE_LOGFMT(LogPlainPropsUObject, Display, "");
+					UE_LOGFMT(LogPlainPropsUObject, Display, "Detected diff in package {PackageName}:\n{Diff}",
+						PackageNames[I], Diff.ToString());
 					++NumDiffs;
+
+					{
+						TRACE_CPUPROFILER_EVENT_SCOPE(PrintToYaml);
+						FPlainPropsLinkerSandbox::PrintPackageToYamlFile(PackageNames[I]);
+					}
+
+					{
+						TRACE_CPUPROFILER_EVENT_SCOPE(DumpToJson);
+						const FDiffObjectNode& RootNode = Ctx.Diffs.Last();
+						const UObject* A = static_cast<const UObject*>(RootNode.A);
+						const UObject* B = static_cast<const UObject*>(RootNode.B);
+						if (A)
+						{
+							DumpUObjectToJsonFile(A);
+						}
+						if (B)
+						{
+							DumpUObjectToJsonFile(B);
+						}
+					}
 				}
-			}
-			UE_LOGFMT(LogPlainPropsUObject, Display, "Detected {NumDiffs} diffs in {SaveNum} packages\n{Diffs}",
-				NumDiffs, Packages.Num(), Diffs.ToString());
+				else if (Ctx.NumIgnoredDiffs)
+				{
+					++NumIgnored;
+				}
+			}, Flags);
+			UE_LOGFMT(LogPlainPropsUObject, Display, "");
+			UE_LOGFMT(LogPlainPropsUObject, Display, "Ignored {NumIgnored} package diff(s) in {Num} package(s)", NumIgnored, NumPackages);
+			UE_LOGFMT(LogPlainPropsUObject, Display, "Detected {NumDiffs} package diff(s) in {Num} package(s)", NumDiffs, NumPackages);
 		}
 	}
 
+	TRACE_CPUPROFILER_EVENT_FLUSH();
 	TRACE_CPUPROFILER_EVENT_SCOPE(Shutdown);
+	FPersistentBatchLoader::PrintDebugInfo();
+	
 #if UE_BUILD_DEBUG
 	// Note: Shutdown is too slow to be enabled by default.
 	// It will reset and flush all the linkers which takes ~15s for ~3000 packages in Release config
@@ -1955,7 +2222,6 @@ int32 RoundtripViaLinkerBatch(TConstArrayView<UObject*> Objects, ERoundtrip Opti
 #endif
 
 #endif // PP_WITH_ENGINE_HOOKS
-
 	return 0;
 }
 
@@ -1985,11 +2251,10 @@ bool SaveTestPackage(UPackage* Package, const FString& Filename, const FSavePack
 		// This will cause our package to be considered missing during load (see
 		// TryGetPackageInfoFromPackageStore). As a workaround we create a temporary
 		// empty package file.
-		if (FArchive* Ar = FileManager.CreateFileWriter(*Filename))
+		if (TUniquePtr<FArchive> Ar(FileManager.CreateFileWriter(*Filename)); Ar)
 		{
 			uint32 TempData = 0;
 			Ar->Serialize(&TempData, sizeof(TempData));
-			delete Ar;
 		}
 	}
 
@@ -2001,18 +2266,18 @@ bool SaveTestPackage(UPackage* Package, const FString& Filename, const FSavePack
 
 UPackage* LoadTestPackage(const FString& Filename)
 {
-	UPackage* Result = nullptr;
+	UPackage* Package = nullptr;
 
 #if PP_WITH_ENGINE_HOOKS
-	FScopeLoadTestPackage LoadScope;
+	FScopeLoadTestPackage LoadScope(/*bPlainProps*/true);
 
 	FName PackageName = FName(*Filename);
 
-	FPlainPropsLinkerSandbox::MountPackagesFromMemory(PackageName);
-	Result = FPlainPropsLinkerSandbox::LoadPlainPropsPackageFromMemory(PackageName);
+	Package = FPlainPropsLinkerSandbox::MountPlainPropsPackage(PackageName, /*bClone*/false);
+	FPlainPropsLinkerSandbox::LoadPlainPropsPackage(Package);
 #endif
 
-	return Result;
+	return Package;
 }
 
 } // namespace PlainProps::UE
diff --git a/Engine/Plugins/Experimental/PlainPropsUObject/Source/Private/PlainPropsUObjectDiff.cpp b/Engine/Plugins/Experimental/PlainPropsUObject/Source/Private/PlainPropsUObjectDiff.cpp
--- a/Engine/Plugins/Experimental/PlainPropsUObject/Source/Private/PlainPropsUObjectDiff.cpp
+++ b/Engine/Plugins/Experimental/PlainPropsUObject/Source/Private/PlainPropsUObjectDiff.cpp
@@ -1,12 +1,19 @@
 // Copyright Epic Games, Inc. All Rights Reserved.
 
 #include "PlainPropsUObjectDiffInternal.h"
+#include "PlainPropsInternalPrivateMemberPtr.h"
 #include "UObject/Class.h"
 #include "UObject/Object.h"
 #include "UObject/Package.h"
 #include "UObject/UnrealType.h"
 #include "Containers/StringFwd.h"
 
+// Temp hacks for non-intrusive prototype
+PP_DEFINE_PRIVATE_MEMBER_NAMESPACE_BEGIN(FFieldPath_Private)
+PP_DEFINE_PRIVATE_MEMBER_PTR(FFieldPath, ResolvedOwner,	TWeakObjectPtr<UStruct>);
+PP_DEFINE_PRIVATE_MEMBER_PTR(FFieldPath, Path,			TArray<FName>);
+PP_DEFINE_PRIVATE_MEMBER_NAMESPACE_END()
+
 namespace PlainProps::UE
 {
 
@@ -54,6 +61,99 @@ FAnsiStringView ToString(EDiffObjectNodeType Type)
 	return Types[(uint8)Type];
 }
 
+// Slow FName based IsChildOf implementation that does not require any dependencies
+static bool IsChildOf(const UStruct* Struct, FName BaseName)
+{
+	for (const UStruct* It : Struct->GetSuperStructIterator())
+	{
+		if (It->GetFName() == BaseName)
+		{
+			return true;
+		}
+	}
+	return false;
+}
+
+static bool IgnoreDiff(const void* A, const void* B, const FProperty* Property, FDiffObjectContext& Ctx)
+{
+	const FDiffObjectFilter& Flt = Ctx.Filter;
+	const uint64 CastFlags = Property->GetCastFlags();
+	const FName PropertyName = Property->GetFName();
+	const UStruct* OwnerStruct = Property->GetOwnerStruct();
+
+	auto ReturnDiff = [&Ctx](bool bDiff) FORCENOINLINE
+	{
+		// place breakpoint(s) here
+		if (bDiff)
+		{
+			++Ctx.NumIgnoredDiffs;
+			return true;
+		}
+		return false;
+	};
+
+	// 1. Silence diffs for FProperty types missing a custom PlainPropsUObjectDiff implementation
+
+	if (CastFlags & Flt.IgnoreCastFlags)
+	{
+		return ReturnDiff(true);
+	}
+
+	// 2. Silence diffs for missing owner structs and unique runtime derived properties
+
+	if (Flt.IgnoreStructs.Contains(OwnerStruct->GetFName()))
+	{
+		return ReturnDiff(true);
+	}
+
+	if (Flt.IgnorePropertiesForStructs.Contains(TPair<FName, FName>{ PropertyName, OwnerStruct->GetFName() }))
+	{
+		return ReturnDiff(true);
+	}
+
+	// 3. Silence diffs for unique runtime derived properties for some specific base classes
+
+	if (const FName* BaseName = Flt.IgnorePropertiesForBases.Find(PropertyName))
+	{
+		return ReturnDiff(IsChildOf(OwnerStruct, *BaseName));
+	}
+
+	// 4. Silence diffs for missing structs in struct properties (inside containers)
+	
+	TArray<FName, TInlineAllocator<2>> StructNames;
+	if (const FStructProperty* StructProp = CastField<FStructProperty>(Property))
+	{
+		StructNames.Add(StructProp->Struct->GetFName());
+	}
+	else if (const FArrayProperty* ArrayProp = CastField<FArrayProperty>(Property))
+	{
+		if (const FStructProperty* InnerProp = CastField<FStructProperty>(ArrayProp->Inner))
+		{
+			StructNames.Add(InnerProp->Struct->GetFName());
+		}
+	}
+	else if (const FMapProperty* MapProp = CastField<FMapProperty>(Property))
+	{
+		if (const FStructProperty* KeyProp = CastField<FStructProperty>(MapProp->GetKeyProperty()))
+		{
+			StructNames.Add(KeyProp->Struct->GetFName());
+		}
+		if (const FStructProperty* ValueProp = CastField<FStructProperty>(MapProp->GetValueProperty()))
+		{
+			StructNames.Add(ValueProp->Struct->GetFName());
+		}
+	}
+	for (FName StructName : StructNames)
+	{
+		if (Flt.IgnoreStructs.Contains(StructName))
+		{
+			return ReturnDiff(true);
+		}
+	}
+
+	return ReturnDiff(false);
+};
+
 ///////////////////////////////////////////////////////////////////////////////
 
 bool DiffStructs(   const void* A, const void* B, const UStruct*   Struct,   FDiffObjectContext& Ctx);
@@ -69,10 +169,18 @@ bool DiffObjectReferences(const UObject* A, const UObject* B, FDiffObjectContext
 	const UClass* Class = A ? A->GetClass() : B->GetClass();
 	if (!A || !B)
 	{
-		Ctx.Diffs.Emplace(EDiffObjectNodeType::Object, A, B, Class);
+		// place breakpoint here
+		Ctx.Diffs.Emplace(!A ? EDiffObjectNodeType::NotInA : EDiffObjectNodeType::NotInB, A, B, Class);
 		return true;
 	}
 
+	bool bAlreadyInSet;
+	Ctx.Visited.Add({A, B}, &bAlreadyInSet);
+	if (bAlreadyInSet)
+	{
+		return false;
+	}
+
 	if (Class != B->GetClass())
 	{
 		// Treat dynamic (non-native) classes as equal if their properties are equal.
@@ -80,7 +188,7 @@ bool DiffObjectReferences(const UObject* A, const UObject* B, FDiffObjectContext
 			B->GetClass()->HasAnyClassFlags(CLASS_Native) ||
 			DiffObjectReferences(Class, B->GetClass(), Ctx))
 		{
-			Ctx.Diffs.Emplace(EDiffObjectNodeType::Class, A, B);
+			Ctx.Diffs.Emplace(EDiffObjectNodeType::Class, A, B, Class);
 			return true;
 		}
 	}
@@ -96,6 +204,7 @@ bool DiffObjectReferences(const UObject* A, const UObject* B, FDiffObjectContext
 		bool bDiff = FStringView(PathA) != FStringView(PathB);
 		if (bDiff)
 		{
+			// place breakpoint here
 			Ctx.Diffs.Emplace(EDiffObjectNodeType::Object, A, B, Class);
 		}
 		return bDiff;
@@ -164,7 +273,7 @@ bool DiffObjects(const UObject* A, const UObject* B, FDiffObjectContext& Ctx)
 			B->GetClass()->HasAnyClassFlags(CLASS_Native) ||
 			DiffObjectReferences(Class, B->GetClass(), Ctx))
 		{
-			Ctx.Diffs.Emplace(EDiffObjectNodeType::Class, A, B);
+			Ctx.Diffs.Emplace(EDiffObjectNodeType::Class, A, B, Class);
 			return true;
 		}
 	}
@@ -191,19 +300,25 @@ bool DiffStructProperties(const void* A, const void* B, const FStructProperty* P
 	// see ScriptStruct->CompareScriptStruct(A, B, /*PortFlags*/ 0);
 	if (ScriptStruct->StructFlags & STRUCT_IdenticalNative)
 	{
-		bool bResult;
-		if (ScriptStruct->GetCppStructOps()->Identical(A, B, /*PortFlags*/0, bResult))
+		if (Ctx.Filter.BypassNativeIdenticalStructs.Contains(ScriptStruct->GetFName()))
+		{
+			; // fall through to default struct diffing doing one property at a time
+		}
+		else
 		{
-			if (!bResult)
+			bool bResult;
+			if (ScriptStruct->GetCppStructOps()->Identical(A, B, /*PortFlags*/0, bResult))
 			{
-				// place breakpoint here
-				if (Ctx.IgnoreDiff(A, B, Property))
+				if (!bResult)
 				{
-					return false;
+					if (IgnoreDiff(A, B, Property, Ctx))
+					{
+						return false;
+					}
+					Ctx.Diffs.Emplace(EDiffObjectNodeType::Property, A, B, Property);
 				}
-				Ctx.Diffs.Emplace(EDiffObjectNodeType::Property, A, B, Property);
+				return !bResult;
 			}
-			return !bResult;
 		}
 	}
 	if (DiffStructs(A, B, ScriptStruct, Ctx))
@@ -233,6 +348,10 @@ bool DiffArrayProperties(const void* A, const void* B, const FArrayProperty* Pro
 
 	if (ArrayA.Num() != ArrayB.Num())
 	{
+		if (IgnoreDiff(A, B, Property, Ctx))
+		{
+			return false;
+		}
 		Ctx.Diffs.Emplace(EDiffObjectNodeType::Property, A, B, Property);
 		return true;
 	}
@@ -257,6 +376,10 @@ bool DiffSetProperties(const void* A, const void* B, const FSetProperty* Propert
 	FScriptSetHelper SetB(Property, B);
 	if (SetA.Num() != SetB.Num())
 	{
+		if (IgnoreDiff(A, B, Property, Ctx))
+		{
+			return false;
+		}
 		Ctx.Diffs.Emplace(EDiffObjectNodeType::Property, A, B, Property);
 		return true;
 	}
@@ -283,6 +406,10 @@ bool DiffMapProperties(const void* A, const void* B, const FMapProperty* Propert
 	FScriptMapHelper MapB(Property, B);
 	if (MapA.Num() != MapB.Num())
 	{
+		if (IgnoreDiff(A, B, Property, Ctx))
+		{
+			return false;
+		}
 		Ctx.Diffs.Emplace(EDiffObjectNodeType::Property, A, B, Property);
 		return true;
 	}
@@ -336,6 +463,31 @@ bool DiffDelegateProperties(const void* A, const void* B, const FDelegatePropert
 	return false;
 }
 
+bool DiffFieldPathProperties(const void* A, const void* B, const FFieldPathProperty* Property, FDiffObjectContext& Ctx)
+{
+	using namespace FFieldPath_Private;
+	const FFieldPath& ValueA = FFieldPathProperty::GetPropertyValue(A);
+	const FFieldPath& ValueB = FFieldPathProperty::GetPropertyValue(B);
+	if (ValueA.*_Path != ValueB.*_Path)
+	{
+		if (IgnoreDiff(A, B, Property, Ctx))
+		{
+			return false;
+		}
+		Ctx.Diffs.Emplace(EDiffObjectNodeType::Property, A, B, Property);
+		return true;
+	}
+
+	const UObject* OwnerA = (ValueA.*_ResolvedOwner).Get();
+	const UObject* OwnerB = (ValueB.*_ResolvedOwner).Get();
+	if (DiffObjectReferences(OwnerA, OwnerB, Ctx))
+	{
+		Ctx.Diffs.Emplace(EDiffObjectNodeType::Property, A, B, Property);
+		return true;
+	}
+	return false;
+}
+
 bool DiffProperties(const void* A, const void* B, const FProperty* Property, FDiffObjectContext& Ctx)
 {
 	uint64 Flags = Property->GetCastFlags();
@@ -389,12 +541,18 @@ bool DiffProperties(const void* A, const void* B, const FProperty* Property, FDi
 			return true;
 		}
 	}
+	else if (HasAnyCastClassFlags<CASTCLASS_FFieldPathProperty>(Flags))
+	{
+		if (DiffFieldPathProperties(A, B, static_cast<const FFieldPathProperty*>(Property), Ctx))
+		{
+			return true;
+		}
+	}
 	else
 	{
 		if (!Property->Identical(A, B, PPF_DeepComparison))
 		{
-			// place breakpoint here
-			if (Ctx.IgnoreDiff(A, B, Property))
+			if (IgnoreDiff(A, B, Property, Ctx))
 			{
 				return false;
 			}
@@ -408,7 +566,6 @@ bool DiffProperties(const void* A, const void* B, const FProperty* Property, FDi
 bool DiffStructs(const void* A, const void* B, const UStruct* Struct, FDiffObjectContext& Ctx)
 {
 	const static FName ClassDefaultObject{"ClassDefaultObject"};
-	const bool bIsUClass = Struct->IsChildOf<UClass>();
 	for (FProperty* Property = Struct->PropertyLink; Property; Property = Property->PropertyLinkNext)
 	{
 		if (!Property->ShouldDuplicateValue()) //Should the property be compared at all?
@@ -416,7 +573,7 @@ bool DiffStructs(const void* A, const void* B, const UStruct* Struct, FDiffObjec
 			continue;
 		}
 		// don't compare the CDO property for classes to prevent infinite recursion
-		if (bIsUClass && Property->GetName() == ClassDefaultObject)
+		if (Property->GetFName() == ClassDefaultObject && Struct->IsChildOf<UClass>())
 		{
 			continue;
 		}
@@ -458,6 +615,7 @@ bool DiffObjects(FDiffObjectContext& Ctx)
 			return true;
 		}
 
+		Ctx.Visited.Reset();
 		if (DiffObjects(A, *B, Ctx))
 		{
 			return true;
@@ -470,7 +628,7 @@ bool DiffObjects(FDiffObjectContext& Ctx)
 		const UObject** A = ObjectsA.Find(PairB.Key);
 		if (!A)
 		{
-			Ctx.Diffs.Emplace(EDiffObjectNodeType::NotInA, B, nullptr, B->GetClass());
+			Ctx.Diffs.Emplace(EDiffObjectNodeType::NotInA, nullptr, B, B->GetClass());
 			return true;
 		}
 	}
@@ -481,14 +639,11 @@ bool DiffObjects(FDiffObjectContext& Ctx)
 
 void PrintDiff(FUtf8StringBuilderBase& Out, FDiffObjectContext& Ctx)
 {
+	check(Ctx.Diffs.Num());
 	const FDiffObjectNode& LeafNode = Ctx.Diffs[0];
 	const FDiffObjectNode& RootNode = Ctx.Diffs.Last();
 
-	// Append diff type
-	Out << ToString(LeafNode.Type) << ": ";
-
-	// Build ('ValueA' != 'ValueB') string from the first/innermost node
-	Out << "('";
+	// Build DiffType:('ValueA' != 'ValueB') string from the first/innermost node
 	if (LeafNode.Type == EDiffObjectNodeType::Property)
 	{
 		const FProperty* Property = LeafNode.GetProperty();
@@ -496,12 +651,17 @@ void PrintDiff(FUtf8StringBuilderBase& Out, FDiffObjectContext& Ctx)
 		FString B;
 		Property->ExportTextItem_Direct(A, LeafNode.A, nullptr, (UObject*)Ctx.RootA, PPF_None, nullptr);
 		Property->ExportTextItem_Direct(B, LeafNode.B, nullptr, (UObject*)Ctx.RootB, PPF_None, nullptr);
+		Out << Property->GetClass()->GetFName();
+		Out << ": ('";
 		A.ReplaceInline(TEXT("\r\n"), TEXT(" "));
 		B.ReplaceInline(TEXT("\r\n"), TEXT(" "));
 		Out << A << "' != '" << B;
 	}
 	else
 	{
+		// Append diff type
+		Out << ToString(LeafNode.Type);
+		Out << ": ('";
 		const UObject* A = static_cast<const UObject*>(LeafNode.A);
 		const UObject* B = static_cast<const UObject*>(LeafNode.B);
 		Out << (A ? A->GetFullName() : "<null>");
@@ -528,7 +688,7 @@ void PrintDiff(FUtf8StringBuilderBase& Out, FDiffObjectContext& Ctx)
 	bool bPrintPropertyPath = true;
 	for (const FDiffObjectNode& Diff : ReverseIterate(Ctx.Diffs))
 	{
-		if (Diff.Type != EDiffObjectNodeType::Property)
+		if (Diff.Type != EDiffObjectNodeType::Property && Diff.Type != EDiffObjectNodeType::Class)
 		{
 			// all intermediate diff nodes should be properties
 			check(&Diff == &LeafNode || &Diff == &RootNode);
@@ -538,13 +698,22 @@ void PrintDiff(FUtf8StringBuilderBase& Out, FDiffObjectContext& Ctx)
 		// skip for inner properties coming after a range property, they just repeat the same type info again
 		if (bPrintPropertyPath)
 		{
-			const FProperty* Property = Diff.GetProperty();
-			FString ExtendedTypeText;
-			FString TypeText = Property->GetCPPType(&ExtendedTypeText);
-			Out << '.';
-			Out << Property->GetFName();
-			Out << ':';
-			Out << TypeText << ExtendedTypeText;
+			if (Diff.Type == EDiffObjectNodeType::Property)
+			{
+				const FProperty* Property = Diff.GetProperty();
+				FString ExtendedTypeText;
+				FString TypeText = Property->GetCPPType(&ExtendedTypeText);
+				Out << '.';
+				Out << Property->GetFName();
+				Out << ':';
+				Out << TypeText << ExtendedTypeText;
+			}
+			else
+			{
+				check(Diff.Type == EDiffObjectNodeType::Class);
+				Out << ".Class:";
+				Out << Diff.GetClass()->GetFName();
+			}
 		}
 		// Append range index
 		if (Diff.Idx != ~0ull)
diff --git a/Engine/Plugins/Experimental/PlainPropsUObject/Source/Private/PlainPropsUObjectDiffInternal.h b/Engine/Plugins/Experimental/PlainPropsUObject/Source/Private/PlainPropsUObjectDiffInternal.h
--- a/Engine/Plugins/Experimental/PlainPropsUObject/Source/Private/PlainPropsUObjectDiffInternal.h
+++ b/Engine/Plugins/Experimental/PlainPropsUObject/Source/Private/PlainPropsUObjectDiffInternal.h
@@ -8,8 +8,10 @@
 #pragma once
 
 #include "Containers/Array.h"
+#include "Containers/Map.h"
 #include "Containers/StringFwd.h"
 #include "Templates/Function.h"
+#include "UObject/NameTypes.h"
 
 class FProperty;
 class UClass;
@@ -51,17 +53,33 @@ struct FDiffObjectNode
 	}
 };
 
-struct FDiffObjectContext
+struct FDiffObjectFilter
 {
-	using FIgnoreFilter = TFunction<bool(const void* A, const void* B, const FProperty*)>;
+	// Structs for which to bypass the native identical function
+	TSet<FName> BypassNativeIdenticalStructs;
+	// Structs for which to ignore any diffs
+	TSet<FName> IgnoreStructs;
+	// Struct specific properties for which to ignore any diffs
+	TSet<TPair<FName, FName>> IgnorePropertiesForStructs;
+	// Base class properties for which to ignore any diffs (more expensive to check)
+	TMap<FName, FName> IgnorePropertiesForBases;
+	// Cast flags for properties for which to ignore any diffs
+	uint64 IgnoreCastFlags = 0;//CASTCLASS_FMulticastInlineDelegateProperty|CASTCLASS_FMulticastSparseDelegateProperty;
+};
 
+struct FDiffObjectContext
+{
 	// in
 	const UObject* const RootA = nullptr;
 	const UObject* const RootB = nullptr;
-	FIgnoreFilter IgnoreDiff;
+	const FDiffObjectFilter& Filter;
+
+	// temporary data structures	
+	TSet<TPair<const UObject*, const UObject*>> Visited;
 
 	// out
 	TArray<FDiffObjectNode> Diffs;
+	uint32 NumIgnoredDiffs = 0;
 };
 
 bool DiffObjects(FDiffObjectContext& Ctx);
diff --git a/Engine/Plugins/Experimental/PlainPropsUObject/Source/Public/PlainPropsRoundtripTest.h b/Engine/Plugins/Experimental/PlainPropsUObject/Source/Public/PlainPropsRoundtripTest.h
--- a/Engine/Plugins/Experimental/PlainPropsUObject/Source/Public/PlainPropsRoundtripTest.h
+++ b/Engine/Plugins/Experimental/PlainPropsUObject/Source/Public/PlainPropsRoundtripTest.h
@@ -3,8 +3,9 @@
 #pragma once
 
 #include "Containers/ArrayView.h"
+#include "Containers/Map.h"
 #include "Misc/EnumClassFlags.h"
-#include "Templates/Function.h"
+#include "UObject/NameTypes.h"
 
 class UObject;
 class UPackage;
@@ -28,11 +29,29 @@ enum class ERoundtrip : uint8
 };
 ENUM_CLASS_FLAGS(ERoundtrip);
 
-using FLinkerDiffIgnoreFilter = TFunction<bool(const void* A, const void* B, const FProperty*)>;
+struct FLinkerDiffFilter
+{
+	// Structs for which to bypass the native identical function
+	TSet<FName> BypassNativeIdenticalStructs;
+	// Structs for which to ignore any diffs
+	TSet<FName> IgnoreStructs;
+	// Struct specific properties for which to ignore any diffs
+	TSet<TPair<FName, FName>> IgnorePropertiesForStructs;
+	// Base class properties for which to ignore any diffs (more expensive to check)
+	TMap<FName, FName> IgnorePropertiesForBases;
+	// Cast flags for properties for which to ignore any diffs
+	uint64 IgnoreCastFlags = 0;
+};
+
+using FDropPackagesFunc = void (*)(TArray<UPackage*>&&, int32);
 
 PLAINPROPSUOBJECT_API void SchemaBindAllTypes(EBindMode Mode, EBatchType BatchType);
 PLAINPROPSUOBJECT_API int32 RoundtripViaPlainBatch(TConstArrayView<UObject*> Objects, ERoundtrip Options);
-PLAINPROPSUOBJECT_API int32 RoundtripViaLinkerBatch(TConstArrayView<UObject*> Objects, ERoundtrip Options, FLinkerDiffIgnoreFilter DiffIgnoreFilter);
+PLAINPROPSUOBJECT_API int32 RoundtripViaLinkerBatch(
+	TConstArrayView<UObject*> Objects,
+	ERoundtrip Options,
+	const FLinkerDiffFilter& DiffFilter,
+	FDropPackagesFunc DropPackagesFunc);
 
 PLAINPROPSUOBJECT_API bool SaveTestPackage(UPackage* Package, const FString& Filename, const FSavePackageArgs& SaveArgs);
 PLAINPROPSUOBJECT_API UPackage* LoadTestPackage(const FString& Filename);
