PlainProps: Add missing include to PlainPropsUObjectDiff.cpp

#rnx

#Synced-CL 48718024

diff --git a/Engine/Plugins/Experimental/PlainPropsUObject/Source/Private/PlainPropsUObjectDiff.cpp b/Engine/Plugins/Experimental/PlainPropsUObject/Source/Private/PlainPropsUObjectDiff.cpp
--- a/Engine/Plugins/Experimental/PlainPropsUObject/Source/Private/PlainPropsUObjectDiff.cpp
+++ b/Engine/Plugins/Experimental/PlainPropsUObject/Source/Private/PlainPropsUObjectDiff.cpp
@@ -3,6 +3,7 @@
 #include "PlainPropsUObjectDiffInternal.h"
 #include "PlainPropsInternalPrivateMemberPtr.h"
 #include "UObject/Class.h"
+#include "UObject/FieldPath.h"
 #include "UObject/Object.h"
 #include "UObject/Package.h"
 #include "UObject/UnrealType.h"
