Added option to gather all differences during diffing.
Added an ignore list on the PlainBatch RoundTrip test to ignore diffs encountered during the TPS serialization.

#rnx

#Synced-CL 48851610

diff --git a/Engine/Plugins/Experimental/PlainProps/Source/Private/PlainPropsDiff.cpp b/Engine/Plugins/Experimental/PlainProps/Source/Private/PlainPropsDiff.cpp
--- a/Engine/Plugins/Experimental/PlainProps/Source/Private/PlainPropsDiff.cpp
+++ b/Engine/Plugins/Experimental/PlainProps/Source/Private/PlainPropsDiff.cpp
@@ -11,6 +11,21 @@
 namespace PlainProps
 {
 
+int32 FDiffContext::AddDiff(FDiffNode&& Node)
+{
+	int32 Index = Diffs.Add(Node);
+
+	LastSibling = Index;
+
+	return Index;
+}
+
+void FDiffContext::Reset()
+{
+	Diffs.Reset();
+	LastSibling = INDEX_NONE;
+}
+
 inline bool DiffMember(FLeafMemberBinding Member, const void* OwnerA, const void* OwnerB, const FBindContext&, FOptionalMemberId)
 {
 	return !!DiffLeaf(At(OwnerA, Member.Offset), At(OwnerB, Member.Offset), Member.Leaf);
@@ -22,12 +37,36 @@ inline bool DiffMember(FLeafMemberBinding Member, const void* OwnerA, const void
 	if (DiffLeaf(A, B, Member.Leaf))
 	{
 		// place breakpoint here
-		Ctx.Out.Add({Member.Leaf.Pack(), Name, {Member.Enum}, A, B});
+		Ctx.AddDiff({Member.Leaf.Pack(), Name, {Member.Enum}, A, B, Ctx.LastSibling});
 		return true;
 	}
 	return false;
 }
 
+template<typename ContextType>
+struct TLastSiblingScope
+{
+	TLastSiblingScope(ContextType& InCtx) : Ctx(InCtx)
+	{
+		if constexpr (std::is_same_v<ContextType, FDiffContext>)
+		{
+			OldLastSibling = InCtx.LastSibling;
+			InCtx.LastSibling = INDEX_NONE;
+		}
+	}
+
+	~TLastSiblingScope()
+	{
+		if constexpr (std::is_same_v<ContextType, FDiffContext>)
+		{
+			Ctx.LastSibling = OldLastSibling;
+		}
+	}
+
+	ContextType& Ctx;
+	int32 OldLastSibling;
+};
+
 template<typename ContextType>
 bool DiffMember(FRangeMemberBinding Member, const void* OwnerA, const void* OwnerB, ContextType& Ctx, FOptionalMemberId Name)
 {
@@ -65,7 +104,7 @@ bool DiffMember(FRangeMemberBinding Member, const void* OwnerA, const void* Owne
 		if (bDiff)
 		{
 			// place breakpoint here
-			Ctx.Out.Add({FMemberBindType(Binding.GetSizeType()), Name, {.Range = Binding}, A, B});
+			Ctx.AddDiff({FMemberBindType(Binding.GetSizeType()), Name, {.Range = Binding}, A, B, Ctx.LastSibling});
 		}
 	}
 	return bDiff;
@@ -96,7 +135,7 @@ static bool DiffMember(FStructMemberBinding Member, const void* OwnerA, const vo
 	if (DiffStruct(Member.Id, A, B, Ctx))
 	{
 		// place breakpoint here
-		Ctx.Out.Add({FMemberBindType(Member.Type), Name, {.Struct = Member.Id}, A, B});
+		Ctx.AddDiff({FMemberBindType(Member.Type), Name, {.Struct = Member.Id}, A, B, Ctx.LastSibling});
 		return true;
 	}
 	return false;
@@ -120,25 +159,31 @@ bool DiffMembers(FBindId Id, const void* A, const void* B, ContextType& Ctx)
 	const FStructDeclaration* Declaration;
 	const FSchemaBinding& Schema = Ctx.Schemas.GetStruct(Id, /* out */ Declaration);
 
+	bool bDiff = false;
+
 	FMemberVisitor It(Schema);
 	if (Schema.HasSuper())
 	{
 		if (DiffMembers(It.GrabSuper(), A, B, Ctx))
 		{
-			return true;
+			bDiff = true;
+			if (StopOnFirstDiff(Ctx)) { return true; }
 		}
 	}
 
+	TLastSiblingScope ScopeLastSibling(Ctx);
+
 	for (FMemberId Name : Declaration->GetMemberOrder())
 	{
 		if (DiffMember(/* in-out */It, A, B, Ctx, Name))
 		{
-			return true;
+			bDiff = true;
+			if (StopOnFirstDiff(Ctx)) { return true; }
 		}
 	}
 
 	check(!It.HasMore());
-	return false;
+	return bDiff;
 }
 
 template<typename ContextType>
@@ -566,4 +611,37 @@ bool DiffStruct(FStructView A, FStructView B, FReadDiffPath& Out)
 	return DiffMembers(A, B, Out);
 }
 
+bool operator==(const FDiffPath& A, const FDiffPath& B)
+{
+	if (A.Num() != B.Num())
+	{
+		return false;
+	}
+	else
+	{
+		for (int32 N = 0; N < A.Num(); ++N)
+		{
+			if (A[N].Name != B[N].Name)
+			{
+				return false;
+			}
+		}
+	}
+
+	return true;
+}
+
+uint32 GetTypeHash(const FDiffPath& Path)
+{
+	uint32 Hash = 0;
+
+	for (const FDiffNode& Node : Path)
+	{
+		uint32 NameHash = GetTypeHash(Node.Name);
+		Hash = HashCombineFast(Hash, NameHash);
+	}
+
+	return Hash;
+}
+
 } // namespace PlainProps
diff --git a/Engine/Plugins/Experimental/PlainProps/Source/Private/PlainPropsInternalDiff.h b/Engine/Plugins/Experimental/PlainProps/Source/Private/PlainPropsInternalDiff.h
--- a/Engine/Plugins/Experimental/PlainProps/Source/Private/PlainPropsInternalDiff.h
+++ b/Engine/Plugins/Experimental/PlainProps/Source/Private/PlainPropsInternalDiff.h
@@ -130,6 +130,7 @@ bool DiffItemSlice(const uint8* A, const uint8* B, uint64 Num, uint32 Stride, Co
 	{
 		if (DiffItem(A, B, Ctx, Member))
 		{
+			// Only gather first item even in EDiffGather::All to avoid a bloated diff.
 			return true;
 		}
 	}
diff --git a/Engine/Plugins/Experimental/PlainProps/Source/Private/PlainPropsPrint.cpp b/Engine/Plugins/Experimental/PlainProps/Source/Private/PlainPropsPrint.cpp
--- a/Engine/Plugins/Experimental/PlainProps/Source/Private/PlainPropsPrint.cpp
+++ b/Engine/Plugins/Experimental/PlainProps/Source/Private/PlainPropsPrint.cpp
@@ -1496,7 +1496,7 @@ void PrintDiff(FUtf8Builder& Out, const FIds& Ids, const FDiffPath& Diff)
 {
 	check(Diff.Num());
 
-	for (FDiffNode Node : ReverseIterate(Diff))
+	for (FDiffNode Node : Diff)
 	{
 		Ids.AppendString(Out, Node.Name);
 		Out << '.';
@@ -1504,7 +1504,7 @@ void PrintDiff(FUtf8Builder& Out, const FIds& Ids, const FDiffPath& Diff)
 	Out.RemoveSuffix(1);
 	Out << ' ';
 	Out << '(';
-	for (FDiffNode Node : ReverseIterate(Diff))
+	for (FDiffNode Node : Diff)
 	{
 		if (Node.Type.IsStruct())
 		{
@@ -1528,6 +1528,16 @@ void PrintDiff(FUtf8Builder& Out, const FIds& Ids, const FDiffPath& Diff)
 	Out << ')';
 }
 
+void PrintDiffs(FUtf8Builder& Out, const FIds& Ids, TConstArrayView<FDiffPath> Diffs)
+{
+	check(Diffs.Num());
+
+	for (const FDiffPath& Diff : Diffs)
+	{
+		PrintDiff(Out, Ids, Diff);
+	}
+}
+
 void PrintDiff(FUtf8Builder& Out, const FBatchIds& Ids, const FReadDiffPath& Diff)
 {
 	check(Diff.Num());
diff --git a/Engine/Plugins/Experimental/PlainProps/Source/Private/PlainPropsUeCoreBindings.cpp b/Engine/Plugins/Experimental/PlainProps/Source/Private/PlainPropsUeCoreBindings.cpp
--- a/Engine/Plugins/Experimental/PlainProps/Source/Private/PlainPropsUeCoreBindings.cpp
+++ b/Engine/Plugins/Experimental/PlainProps/Source/Private/PlainPropsUeCoreBindings.cpp
@@ -90,26 +90,26 @@ void FTransformBinding::Save(FMemberBuilder& Dst, const FTransform& Src, const F
 	}
 }
 
-inline FDiffNode MakeStructDiff(const void* A, const void* B, FMemberId Name, FBindId Id)
+inline FDiffNode MakeStructDiff(const void* A, const void* B, FMemberId Name, FBindId Id, FDiffContext& Ctx)
 {
-	return { DefaultStructBindType, Name, {.Struct = Id}, A, B };
+	return { DefaultStructBindType, Name, {.Struct = Id}, A, B, Ctx.LastSibling };
 }
 
 bool FTransformBinding::Diff(const FTransform& A, const FTransform& B, FDiffContext& Ctx) const
 {
 	if (!A.TranslationEquals(B, 0.0))
 	{
-		Ctx.Out.Emplace(MakeStructDiff(&A, &B, MemberIds[(uint8)EMember::Translate], VectorId));
+		Ctx.AddDiff(MakeStructDiff(&A, &B, MemberIds[(uint8)EMember::Translate], VectorId, Ctx));
 		return true;
 	}
-	else if (!A.RotationEquals(B, 0.0))
+	if (!A.RotationEquals(B, 0.0))
 	{
-		Ctx.Out.Emplace(MakeStructDiff(&A, &B, MemberIds[(uint8)EMember::Rotate], QuatId));
+		Ctx.AddDiff(MakeStructDiff(&A, &B, MemberIds[(uint8)EMember::Rotate], QuatId, Ctx));
 		return true;
 	}
-	else if (!A.Scale3DEquals(B, 0.0))
+	if (!A.Scale3DEquals(B, 0.0))
 	{
-		Ctx.Out.Emplace(MakeStructDiff(&A, &B, MemberIds[(uint8)EMember::Scale], VectorId));
+		Ctx.AddDiff(MakeStructDiff(&A, &B, MemberIds[(uint8)EMember::Scale], VectorId, Ctx));
 		return true;
 	}
 
diff --git a/Engine/Plugins/Experimental/PlainProps/Source/Public/PlainPropsBind.h b/Engine/Plugins/Experimental/PlainProps/Source/Public/PlainPropsBind.h
--- a/Engine/Plugins/Experimental/PlainProps/Source/Public/PlainPropsBind.h
+++ b/Engine/Plugins/Experimental/PlainProps/Source/Public/PlainPropsBind.h
@@ -261,7 +261,7 @@ struct ICustomBinding
 	// Overload to load via a single memcpy call instead of LoadCustom, Out.Offset is 0 by default
 	virtual void	PlanCustom(FMemcpyLoadPlan& Out) const {}
 
-	// Overload to track the first diffing member, see FDiffPath
+	// Overload to track the first or all diffing members, see FDiffTree
 	PLAINPROPS_API virtual bool	DiffCustom(const void* StructA, const void* StructB, FDiffContext& Ctx) const;
 	
 };
diff --git a/Engine/Plugins/Experimental/PlainProps/Source/Public/PlainPropsDiff.h b/Engine/Plugins/Experimental/PlainProps/Source/Public/PlainPropsDiff.h
--- a/Engine/Plugins/Experimental/PlainProps/Source/Public/PlainPropsDiff.h
+++ b/Engine/Plugins/Experimental/PlainProps/Source/Public/PlainPropsDiff.h
@@ -22,20 +22,40 @@ struct FDiffNode
 	FDiffMetadata			Meta;
 	const void*				A;
 	const void*				B;
+	int32					PrecedingSibling = INDEX_NONE; // Index of the previous sibling of the node in the tree chain.
 };
 
-struct FDiffPath : public TArray<FDiffNode, TInlineAllocator<16>> {};
+// Flat tree of nodes. Nodes are linked via the index to a previous sibling.
+// A tree can contain multiple diff paths.
+struct FDiffTree : public TArray<FDiffNode, TInlineAllocator<16>> {};
+
+enum EDiffGather { First, All };
 
 // Tracks diff path for diff tools unlike the const FBindContext& overloads for delta saving
 struct FDiffContext : FBindContext
 {
-	FDiffPath Out;
+	EDiffGather		GatherMode = EDiffGather::First;
+	FDiffTree		Diffs;
+	int32			LastSibling = INDEX_NONE;
+
+	PLAINPROPS_API int32 AddDiff(FDiffNode&& Node);
+	PLAINPROPS_API void Reset();
 };
 
+inline bool StopOnFirstDiff(const FDiffContext& Ctx)
+{
+	return (Ctx.GatherMode == EDiffGather::First);
+}
+
+inline bool StopOnFirstDiff(const FBindContext& Ctx)
+{
+	return true;
+}
+
 ////////////////////////////////////////////////////////////////////////////////
 // Tracking and non-tracking methods to diff member leaves/ranges/structs
 //
-// FDiffContext& overloads track inner FDiffPath, caller must add outermost FDiffNode
+// FDiffContext& overloads track inner FDiffTree, caller must add outermost FDiffNode
 
 [[nodiscard]] PLAINPROPS_API bool DiffStructs(const void* A, const void* B, FBindId Id, const FBindContext& Ctx);
 [[nodiscard]] PLAINPROPS_API bool DiffStructs(const void* A, const void* B, FBindId Id, FDiffContext& Ctx);
@@ -65,4 +85,16 @@ struct FReadDiffPath : public TArray<FReadDiffNode, TInlineAllocator<16>> {};
 bool PLAINPROPS_API DiffStruct(FStructView A, FStructView B, FReadDiffPath& Out);
 bool PLAINPROPS_API DiffSchemas(FSchemaBatchId A, FSchemaBatchId B);
 
+////////////////////////////////////////////////////////////////////////////////
+
+// Struct for converting diff results to a less efficient but easier to
+// process structure.
+// Nodes in the path are stored in forward order. Sibling links are invalidated.
+struct FDiffPath : public TArray<FDiffNode> 
+{
+	[[nodiscard]] PLAINPROPS_API friend uint32 GetTypeHash(const FDiffPath& Path);
+};
+
+PLAINPROPS_API bool operator==(const FDiffPath& A, const FDiffPath& B);
+
 } // namespace PlainProps
diff --git a/Engine/Plugins/Experimental/PlainProps/Source/Public/PlainPropsPrint.h b/Engine/Plugins/Experimental/PlainProps/Source/Public/PlainPropsPrint.h
--- a/Engine/Plugins/Experimental/PlainProps/Source/Public/PlainPropsPrint.h
+++ b/Engine/Plugins/Experimental/PlainProps/Source/Public/PlainPropsPrint.h
@@ -60,5 +60,6 @@ PLAINPROPS_API void PrintMemberSchema(FUtf8Builder& Out, const FIds& Ids, FMembe
 PLAINPROPS_API void PrintYamlBatch(FUtf8Builder& Out, const FBatchIds& Ids, TConstArrayView<FStructView> Objects);
 PLAINPROPS_API void PrintDiff(FUtf8Builder& Out, const FIds& Ids, const FDiffPath& Diff);
 PLAINPROPS_API void PrintDiff(FUtf8Builder& Out, const FBatchIds& Ids, const FReadDiffPath& Diff);
+PLAINPROPS_API void PrintDiffs(FUtf8Builder& Out, const FIds& Ids, TConstArrayView<FDiffPath> Diffs);
 
 } // namespace PlainProps
diff --git a/Engine/Plugins/Experimental/PlainProps/Source/Public/PlainPropsSaveMember.h b/Engine/Plugins/Experimental/PlainProps/Source/Public/PlainPropsSaveMember.h
--- a/Engine/Plugins/Experimental/PlainProps/Source/Public/PlainPropsSaveMember.h
+++ b/Engine/Plugins/Experimental/PlainProps/Source/Public/PlainPropsSaveMember.h
@@ -21,7 +21,7 @@ inline bool Track(bool bDiff, FDiffContext& Ctx, FMemberBindType Type, FMemberId
 {
 	if (bDiff)
 	{
-		Ctx.Out.Emplace(Type, ToOptional(Name), Meta, A, B);
+		Ctx.AddDiff({Type, ToOptional(Name), Meta, A, B, Ctx.LastSibling});
 	}
 	
 	return bDiff;
diff --git a/Engine/Plugins/Experimental/PlainProps/Source/Public/PlainPropsUeCoreBindings.h b/Engine/Plugins/Experimental/PlainProps/Source/Public/PlainPropsUeCoreBindings.h
--- a/Engine/Plugins/Experimental/PlainProps/Source/Public/PlainPropsUeCoreBindings.h
+++ b/Engine/Plugins/Experimental/PlainProps/Source/Public/PlainPropsUeCoreBindings.h
@@ -56,7 +56,7 @@ public:
 	inline static FSensitiveName FromUnstableInt(IntType UnstableInt) { return FSensitiveName(UnstableInt); }
 	[[nodiscard]] inline bool operator==(FSensitiveName Other) const { return ToUnstableInt() == Other.ToUnstableInt(); }
 	[[nodiscard]] inline bool operator!=(FSensitiveName Other) const { return !(*this == Other); }
-	[[nodiscard]] friend inline uint32 GetTypeHash(FSensitiveName Name) { return GetTypeHash(Name.ToUnstableInt()); }
+	[[nodiscard]] friend inline uint32 GetTypeHash(FSensitiveName Name) { return ::GetTypeHash(Name.ToUnstableInt()); }
 
 private:
 	inline explicit FSensitiveName(IntType);
diff --git a/Engine/Plugins/Experimental/PlainPropsUObject/Source/Private/PlainPropsRoundtripTest.cpp b/Engine/Plugins/Experimental/PlainPropsUObject/Source/Private/PlainPropsRoundtripTest.cpp
--- a/Engine/Plugins/Experimental/PlainPropsUObject/Source/Private/PlainPropsRoundtripTest.cpp
+++ b/Engine/Plugins/Experimental/PlainPropsUObject/Source/Private/PlainPropsRoundtripTest.cpp
@@ -1614,6 +1614,43 @@
 	return Object->HasAllFlags(RF_ClassDefaultObject) && Object->GetClass()->HasAllClassFlags(CLASS_Native);
 }
 
+using FDiffSet = TSet<FDiffPath>;
+
+FDiffSet Flatten(const FDiffTree& DiffTree)
+{
+	FDiffSet DiffSet;
+
+	FDiffPath CurrentPath;
+	CurrentPath.Reserve(DiffTree.Num());
+
+	for (int32 N = DiffTree.Num() - 1; N >= 0; --N)
+	{
+		FDiffNode Node = DiffTree[N];
+
+		int32 PrecedingSibling = Node.PrecedingSibling;
+
+		// The preceding sibling index has no meaning when we convert into a path.
+		Node.PrecedingSibling = INDEX_NONE;
+		CurrentPath.Add(Node);
+
+		bool bHasPrecedingSibling = (PrecedingSibling != INDEX_NONE);
+
+		// Terminate the path if we have finished a 'run'.
+		if ((N == 0) || bHasPrecedingSibling)
+		{
+			DiffSet.Add(CurrentPath);
+
+			if (bHasPrecedingSibling)
+			{
+				// The previous sibling shares the same parent path as this node.
+				CurrentPath.Pop();
+			}
+		}
+	}
+
+	return DiffSet;
+}
+
 int32 RoundtripViaPlainBatch(TConstArrayView<UObject*> Objects, ERoundtrip Options)
 {
 	TRACE_CPUPROFILER_EVENT_SCOPE(RoundtripViaPlainBatch);
@@ -1693,66 +1730,12 @@
 		check(!Instance.UPS || (Instance.UPS != Instance.Base));
 	}
 
-	// Save
-	UE_LOGFMT(LogPlainPropsUObject, Display, "Saving UObjects to PlainProps with InMemoryNames...");
 	FMemoryBatch Plain;
 	TCustomPropertyBindings<FMemoryPropertyBatch> Customs(Plain.Properties, GUE.Customs);
-	{
-		FBatchSaver Batch(Customs.Overlay, GUObjectArray.GetObjectArrayNum());
-		SavePlainProps(Batch, Instances);
-		Plain.Data = Batch.Write(/* out */ &Plain.RuntimeIds);
 
-		if (EnumHasAnyFlags(Options, ERoundtrip::TextMemory))
-		{
-			FBatchTextRoundtripper MemoryBatch(MakeMemoryView(Plain.Data), ESchemaFormat::InMemoryNames);
-			MemoryBatch.RoundtripTextForDiffTest(Instances);
-		}
-		if (EnumHasAnyFlags(Options, ERoundtrip::TextStable))
-		{
-			UE_LOGFMT(LogPlainPropsUObject, Display, "Saving UObjects to PlainProps with StableNames...");
-			TArray64<uint8> StableData = Batch.Write();
-
-			FBatchTextRoundtripper StableBatch(MakeMemoryView(StableData), ESchemaFormat::StableNames);
-			StableBatch.RoundtripTextForDiffTest(Instances);
-		}
-	}
-
-	// Load
-	uint32 NumPpDiffs = 0;
-	if (EnumHasAnyFlags(Options, ERoundtrip::PP))
-	{
-		UE_LOGFMT(LogPlainPropsUObject, Display, "Loading UObjects from PlainProps...");
-		FMemoryBatchLoader Batch(Customs.Overlay, MakeMemoryView(Plain.Data), Plain.RuntimeIds);
-		LoadPlainProps(Batch, Instances);
-
-		// Diff original vs PlainProps
-		UE_LOGFMT(LogPlainPropsUObject, Display, "Diffing UObjects roundtripped via PlainProps...");
-		FDiffContext DiffCtx = {{ GUE.Schemas, Customs.Overlay }};
-		TUtf8StringBuilder<256> PpDiffPaths;
-		for (FInstance& Instance : Instances)
-		{
-			if (DiffStructs(Instance.Orig, Instance.PP, Instance.Id, /* in-out */ DiffCtx))
-			{
-				PpDiffPaths.Append("\n");
-				GPpDiff.Instances.Add(&Instance);
-				PrintDiff(/* out */ PpDiffPaths, GUE.Names, DiffCtx.Out);
-
-				//Batch.Reload(Instance.UPS, &Instance - Instances.GetData());
-				DiffCtx.Out.Reset();
-				PpDiffPaths.Append(" in ");
-				PpDiffPaths.Append(Instance.Orig->GetFullName());
-				++NumPpDiffs;
-			}
-		}
-		GPpDiff.Str = PpDiffPaths.GetData();
-		UE_LOGFMT(LogPlainPropsUObject, Display, "Detected {Diffs} diffs in {Objs} UObjects saved in a {KB}KB value stream using PlainProps{DiffText}", NumPpDiffs, Instances.Num(), Plain.Data.NumBytes() / 1024, PpDiffPaths.ToString());
-		for (const FInstance* Instance : GPpDiff.Instances)
-		{
-			DumpUObjectToJsonFile(Instance->Orig);
-			DumpUObjectToJsonFile(Instance->PP);
-		}
-	}
-
+	// We want to ignore diffs that happen during TPS serialization when comparing the
+	// results of PlainProps serialization. For this reason we start with TPS. Such
+	// differences could be a result of the way objects are instantiated in the test.
 	FArchivedProperties Tps;
 	if (EnumHasAnyFlags(Options, ERoundtrip::TPS))
 	{
@@ -1769,15 +1752,22 @@
 	}
 
 	static const FName SkipClasses[] = {
-		"BodySetup", // Skips some structs due to native FCollisionResponse::operator==
-		"NiagaraScript", "NiagaraNodeFunctionCall", "NiagaraMeshRendererProperties", // FNiagaraTypeDefinition::Serialize resets ClassStructOrEnum
+	"BodySetup", // Skips some structs due to native FCollisionResponse::operator==
+	"NiagaraScript", "NiagaraNodeFunctionCall", "NiagaraMeshRendererProperties", // FNiagaraTypeDefinition::Serialize resets ClassStructOrEnum
 	};
 
+	const EDiffGather GatherMode = EDiffGather::All;
+
+	// Map: OriginalInstance -> Diff Paths, that we use as an ignore list.
+	// When a diff happens during the original vs TPS comparison we want to
+	// exclude it during the original vs PlainProps comparison.
+	TMap<UObject*, FDiffSet> IgnoreDiffMap;
+
 	if (EnumHasAnyFlags(Options, ERoundtrip::TPS))
 	{
 		// Diff original vs TPS
 		UE_LOGFMT(LogPlainPropsUObject, Display, "Diffing UObjects roundtripped via TPS...");
-		FDiffContext DiffCtx = {{ GUE.Schemas, Customs.Overlay }};
+		FDiffContext DiffCtx = { {GUE.Schemas, Customs.Overlay }, GatherMode };
 		TUtf8StringBuilder<256> TpsDiffs;
 		TArray<int32> TpsDiffIdxs;
 		for (FInstance& Instance : Instances)
@@ -1789,13 +1779,20 @@
 
 			if (DiffStructs(Instance.Orig, Instance.TPS, Instance.Id, /* in-out */ DiffCtx))
 			{
+				check(!DiffCtx.Diffs.IsEmpty());
+
+				FDiffSet DiffSet = Flatten(DiffCtx.Diffs);
+				check((GatherMode == EDiffGather::All) || (DiffSet.Num() == 1));
+
 				GTpsDiff.Instances.Add(&Instance);
-				PrintDiff(/* out */ TpsDiffs, GUE.Names, DiffCtx.Out);
-				
-				DiffCtx.Out.Reset();
+				PrintDiffs(/* out */ TpsDiffs, GUE.Names, DiffSet.Array());
+
+				IgnoreDiffMap.Add(Instance.Orig, MoveTemp(DiffSet));
+
+				DiffCtx.Reset();
 				TpsDiffs.Append(" in ");
 				TpsDiffs.Append(Instance.Orig->GetFullName());
-				TpsDiffs.Append("\n");	
+				TpsDiffs.Append("\n");
 				TpsDiffIdxs.Add(&Instance - &Instances[0]);
 				
 				//FArchivedProperties Tmp;
@@ -1814,6 +1811,82 @@
 		UE_LOGFMT(LogPlainPropsUObject, Display, "Detected {Diffs} diffs in {Objs} UObjects saved in a {KB}KB value stream using TPS", TpsDiffIdxs.Num(), Instances.Num(), Tps.Data.NumBytes() / 1024);
 	}
 
+	// Save
+	UE_LOGFMT(LogPlainPropsUObject, Display, "Saving UObjects to PlainProps with InMemoryNames...");
+	{
+		FBatchSaver Batch(Customs.Overlay, GUObjectArray.GetObjectArrayNum());
+		SavePlainProps(Batch, Instances);
+		Plain.Data = Batch.Write(/* out */ &Plain.RuntimeIds);
 
+		if (EnumHasAnyFlags(Options, ERoundtrip::TextMemory))
+		{
+			FBatchTextRoundtripper MemoryBatch(MakeMemoryView(Plain.Data), ESchemaFormat::InMemoryNames);
+			MemoryBatch.RoundtripTextForDiffTest(Instances);
+		}
+		if (EnumHasAnyFlags(Options, ERoundtrip::TextStable))
+		{
+			UE_LOGFMT(LogPlainPropsUObject, Display, "Saving UObjects to PlainProps with StableNames...");
+			TArray64<uint8> StableData = Batch.Write();
+
+			FBatchTextRoundtripper StableBatch(MakeMemoryView(StableData), ESchemaFormat::StableNames);
+			StableBatch.RoundtripTextForDiffTest(Instances);
+		}
+	}
+
+	// Load
+	uint32 NumPpDiffs = 0;
+	if (EnumHasAnyFlags(Options, ERoundtrip::PP))
+	{
+		UE_LOGFMT(LogPlainPropsUObject, Display, "Loading UObjects from PlainProps...");
+		FMemoryBatchLoader Batch(Customs.Overlay, MakeMemoryView(Plain.Data), Plain.RuntimeIds);
+		LoadPlainProps(Batch, Instances);
+
+		// Diff original vs PlainProps
+		UE_LOGFMT(LogPlainPropsUObject, Display, "Diffing UObjects roundtripped via PlainProps...");
+		FDiffContext DiffCtx = { {GUE.Schemas, Customs.Overlay }, GatherMode };
+		TUtf8StringBuilder<256> PpDiffPaths;
+
+		for (FInstance& Instance : Instances)
+		{
+			if (DiffStructs(Instance.Orig, Instance.PP, Instance.Id, /* in-out */ DiffCtx))
+			{
+				check(!DiffCtx.Diffs.IsEmpty());
+
+				FDiffSet DiffSet = Flatten(DiffCtx.Diffs);
+				check((GatherMode == EDiffGather::All) || (DiffSet.Num() == 1));
+
+				FDiffSet FiltertedDiffs;
+				const FDiffSet* IgnoreDiffs = IgnoreDiffMap.Find(Instance.Orig);
+				if (IgnoreDiffs)
+				{
+					FiltertedDiffs = DiffSet.Difference(*IgnoreDiffs);
+				}
+
+				FDiffSet& Diffs = IgnoreDiffs ? FiltertedDiffs : DiffSet;
+
+				if (!Diffs.IsEmpty())
+				{
+					PpDiffPaths.Append("\n");
+					GPpDiff.Instances.Add(&Instance);
+					PrintDiffs(/* out */ PpDiffPaths, GUE.Names, Diffs.Array());
+
+					//Batch.Reload(Instance.UPS, &Instance - Instances.GetData());
+					PpDiffPaths.Append(" in ");
+					PpDiffPaths.Append(Instance.Orig->GetFullName());
+					++NumPpDiffs;
+				}
+
+				DiffCtx.Reset();
+			}
+		}
+		GPpDiff.Str = PpDiffPaths.GetData();
+		UE_LOGFMT(LogPlainPropsUObject, Display, "Detected {Diffs} diffs in {Objs} UObjects saved in a {KB}KB value stream using PlainProps{DiffText}", NumPpDiffs, Instances.Num(), Plain.Data.NumBytes() / 1024, PpDiffPaths.ToString());
+		for (const FInstance* Instance : GPpDiff.Instances)
+		{
+			DumpUObjectToJsonFile(Instance->Orig);
+			DumpUObjectToJsonFile(Instance->PP);
+		}
+	}
+
 	if (EnumHasAnyFlags(Options, ERoundtrip::UPS))
 	{
@@ -1832,7 +1905,7 @@
 
 		// Diff original vs UPS
 		UE_LOGFMT(LogPlainPropsUObject, Display, "Diffing UObjects roundtripped via UPS...");
-		FDiffContext DiffCtx = {{ GUE.Schemas, Customs.Overlay }};
+		FDiffContext DiffCtx = { {GUE.Schemas, Customs.Overlay }, GatherMode };
 		TUtf8StringBuilder<256> UpsDiffs;
 		TArray<int32> UpsDiffIdxs;
 		for (FInstance& Instance : Instances)
@@ -1844,10 +1917,15 @@
 
 			if (DiffStructs(Instance.Orig, Instance.UPS, Instance.Id, /* in-out */ DiffCtx))
 			{
+				check(!DiffCtx.Diffs.IsEmpty());
+
+				FDiffSet DiffSet = Flatten(DiffCtx.Diffs);
+				check((GatherMode == EDiffGather::All) || (DiffSet.Num() == 1));
+
 				GUpsDiff.Instances.Add(&Instance);
-				PrintDiff(/* out */ UpsDiffs, GUE.Names, DiffCtx.Out);
+				PrintDiffs(/* out */ UpsDiffs, GUE.Names, DiffSet.Array());
 				
-				DiffCtx.Out.Reset();
+				DiffCtx.Reset();
 				UpsDiffs.Append(" in ");
 				UpsDiffs.Append(Instance.Orig->GetFullName());
 				UpsDiffs.Append("\n");	
diff --git a/Engine/Plugins/Experimental/PlainPropsUObject/Source/Private/PlainPropsUObjectBindings.cpp b/Engine/Plugins/Experimental/PlainPropsUObject/Source/Private/PlainPropsUObjectBindings.cpp
--- a/Engine/Plugins/Experimental/PlainPropsUObject/Source/Private/PlainPropsUObjectBindings.cpp
+++ b/Engine/Plugins/Experimental/PlainPropsUObject/Source/Private/PlainPropsUObjectBindings.cpp
@@ -1282,7 +1282,7 @@ inline uint32	LeafPropertyHash(double In)											{ return ::UE::PreciseFPHash
 inline bool		LeafPropertyIdentical(float A, float B)								{ return ::UE::PreciseFPEqual(A, B); }
 inline bool		LeafPropertyIdentical(double A, double B)							{ return ::UE::PreciseFPEqual(A, B); }
 template<typename T>
-inline uint32	LeafPropertyHash(T In) requires (std::is_unsigned_v<T>)				{ return GetTypeHash(In); }
+inline uint32	LeafPropertyHash(T In) requires (std::is_unsigned_v<T>)				{ return ::GetTypeHash(In); }
 template<typename T>
 inline bool		LeafPropertyIdentical(T A, T B) requires (std::is_unsigned_v<T>)	{ return A == B; }
 
