PlainProps: Linker roundtrip test tweaks.
Don't try to roundtrip out-of-date developer packages that fails saving.
Pin all loaded packages and and their objects to prevent them from being destroyed when GC is triggered post bp compilation during UPackage::Save.
Initialize all worlds loaded via LoadPackage calls since this is required for saving verse levels (this initialization and more typically happens from LoadMap code).

#rnx

diff --git a/Engine/Plugins/Experimental/PlainPropsEngine/Source/Private/PlainPropsCommandlets.cpp b/Engine/Plugins/Experimental/PlainPropsEngine/Source/Private/PlainPropsCommandlets.cpp
--- a/Engine/Plugins/Experimental/PlainPropsEngine/Source/Private/PlainPropsCommandlets.cpp
+++ b/Engine/Plugins/Experimental/PlainPropsEngine/Source/Private/PlainPropsCommandlets.cpp
@@ -5,12 +5,14 @@
 #include "PlainPropsRoundtripTest.h"
 #include "AssetRegistry/AssetData.h"
 #include "AssetRegistry/AssetRegistryModule.h"
+#include "Engine/World.h"
 #include "Misc/AsciiSet.h"
 #include "Misc/CommandLine.h"
 #include "Misc/CoreDelegates.h"
 #include "Misc/PackageName.h"
 #include "ProfilingDebugging/CpuProfilerTrace.h"
 #include "UObject/PackageReload.h"
+#include "UObject/UObjectIterator.h"
 
 #if WITH_EDITOR
 #include "Engine/AssetManager.h"
@@ -176,6 +178,34 @@ static void DropPackages(TArray<UPackage*>&& Packages, int32 NewNumber)
 	}
 }
 
+// ALevelEntity::PreSaveFromRoot called during save requires a verse::level_entity
+// that is created from PreRegisterAllComponents.
+static void InitWorlds()
+{
+	TRACE_CPUPROFILER_EVENT_SCOPE(InitWorlds);
+
+	// Code inspired by UCookOnTheFlyServer::ValidateSourcePackage
+	FWorldInitializationValues IVS;
+	IVS.AllowAudioPlayback(false);
+	IVS.RequiresHitProxies(false);
+	IVS.ShouldSimulatePhysics(false);
+	IVS.EnableTraceCollision(true);
+	IVS.SetTransactional(false);
+	IVS.CreateWorldPartition(true);
+	IVS.CreateAISystem(false);
+	IVS.CreateNavigation(false);
+
+	for (TObjectIterator<UWorld> It; It; ++It)
+	{
+		if (!It->bIsWorldInitialized)
+		{
+			UE_LOGFMT(LogPlainPropsEngine, Log, "Init world {World}...", It->GetFName());
+			It->InitWorld(IVS);
+			It->UpdateWorldComponents(true, true);
+		}
+	}
+}
+
 //////////////////////////////////////////////////////////////////////////
 
 UTestPlainPropsCommandlet::UTestPlainPropsCommandlet(const FObjectInitializer& Init)
@@ -255,10 +285,10 @@ int32 UTestPlainPropsCommandlet::Main(const FString& Params)
 		FlushAsyncLoading();
 		(*GFlushStreamingFunc)();
 	}
-	UE_LOGFMT(LogPlainPropsEngine, Display, "Starting test...");
 
-	UE_LOGFMT(LogPlainPropsEngine, Display, "Binding types to PlainProps schemas...");
-	
+	UE_LOGFMT(LogPlainPropsEngine, Display, "Init worlds ...");
+	InitWorlds();
+
 	EBindMode Mode = EBindMode::All;
 	if (Params.Find("-bind=source") != INDEX_NONE)
 	{
@@ -275,6 +305,7 @@ int32 UTestPlainPropsCommandlet::Main(const FString& Params)
 		BatchType = EBatchType::Linker;
 	}
 
+	UE_LOGFMT(LogPlainPropsEngine, Display, "Binding types to PlainProps schemas...");
 	{
 		TRACE_CPUPROFILER_EVENT_SCOPE(BindAllTypes);
 		SchemaBindAllTypes(Mode, BatchType);
@@ -304,6 +335,8 @@ int32 UTestPlainPropsCommandlet::Main(const FString& Params)
 		Options = ERoundtrip::TextMemory | ERoundtrip::TextStable;
 	}
 
+	UE_LOGFMT(LogPlainPropsEngine, Display, "Starting roundtrip test...");
+	
 	// Roundtrip
 	if (BatchType == EBatchType::Linker)
 	{
diff --git a/Engine/Plugins/Experimental/PlainPropsUObject/Source/Private/PlainPropsRoundtripTest.cpp b/Engine/Plugins/Experimental/PlainPropsUObject/Source/Private/PlainPropsRoundtripTest.cpp
--- a/Engine/Plugins/Experimental/PlainPropsUObject/Source/Private/PlainPropsRoundtripTest.cpp
+++ b/Engine/Plugins/Experimental/PlainPropsUObject/Source/Private/PlainPropsRoundtripTest.cpp
@@ -741,6 +741,7 @@ static UObject* MakeEmptyInstance(UObject* Obj, FName Name)
 	Params.Name = Name;
 	Params.SetFlags = Obj->GetFlags();
 	Params.Template = Obj->GetArchetype();
+	Params.ExternalPackage = (Params.SetFlags & RF_HasExternalPackage) ? Obj->GetExternalPackage() : nullptr;
 	Params.bAssumeTemplateIsArchetype = true;
 	Params.bCopyTransientsFromClassDefaults = !(Params.SetFlags & RF_ClassDefaultObject);
 	UObject* Result = StaticConstructObject_Internal(Params);
@@ -1301,6 +1302,28 @@ struct FLinkerPropertyBindings : TCustomPropertyBindings<FLinkerPropertyBatch>
 	TCustomPropertyBinding<FSoftObjectPath, FLinkerPropertyBatch> SoftObjectPath;
 };
 
+static UPackage* PinPackage(UPackage* Package)
+{
+	if (Package)
+	{
+		Package->AddRef();
+	}
+	return Package;
+}
+
+static UPackage* PinObjects(UPackage* Package)
+{
+	if (Package)
+	{
+		ForEachObjectWithPackage(Package, [](UObject* Obj)
+		{
+			Obj->AddRef();
+			return true;
+		});
+	}
+	return Package;
+}
+
 struct FPlainPropsLinkerSandbox
 {
 	static constexpr uint32 LoadFlags = (LOAD_Async | LOAD_NoVerify | LOAD_SkipLoadImportedPackages);
@@ -1435,9 +1458,9 @@ struct FPlainPropsLinkerSandbox
 			}
 			UE_LOGFMT(LogPlainPropsUObject, Log, " Mounting baseline for {Package}...", NewPackageName);
 			TUniquePtr<FMountedLinker> MountedInstance(new FMountedLinker(*SavedData, NewPackageName));
-			UPackage* Package = MountedInstance->Package;
+			UPackage* Out = MountedInstance->Package;
 			MountedBaselinePackages.Emplace(NewPackageName, MoveTemp(MountedInstance));
-			return Package;
+			return PinPackage(Out);
 		}
 		return nullptr;
 	}
@@ -1453,9 +1476,9 @@ struct FPlainPropsLinkerSandbox
 			}
 			UE_LOGFMT(LogPlainPropsUObject, Log, " Mounting plainprops for {Package}...", NewPackageName);
 			TUniquePtr<FMountedLinkerBatch> MountedInstance(new FMountedLinkerBatch(*SavedData, NewPackageName));
-			UPackage* Package = MountedInstance->MountedLinker.Package;
+			UPackage* Out = MountedInstance->MountedLinker.Package;
 			MountedPlainpropsPackages.Add(NewPackageName, MoveTemp(MountedInstance));
-			return Package;
+			return PinPackage(Out);
 		}
 		return nullptr;
 	}
@@ -1474,7 +1497,8 @@ struct FPlainPropsLinkerSandbox
 			FMountedLinker& MountedLinker = **MountedLinkerPtr;
 
 			UE_LOGFMT(LogPlainPropsUObject, Log, " Loading baseline for {Package}...", MountedLinker.PackageName);
-			return LoadPackage(MountedLinker.Package, MountedLinker.PackagePath, LoadFlags);
+			UPackage* Out = LoadPackage(MountedLinker.Package, MountedLinker.PackagePath, LoadFlags);
+			return PinObjects(Out);
 		}
 		return nullptr;
 	}
@@ -1486,7 +1510,8 @@ struct FPlainPropsLinkerSandbox
 			FMountedLinkerBatch& LinkerBatch = **LinkerBatchPtr;
 
 			UE_LOGFMT(LogPlainPropsUObject, Log, " Loading PlainProps for {Package}...", LinkerBatch.MountedLinker.PackageName);
-			return LoadPackage(LinkerBatch.MountedLinker.Package, LinkerBatch.MountedLinker.PackagePath, LoadFlags);
+			UPackage* Out = LoadPackage(LinkerBatch.MountedLinker.Package, LinkerBatch.MountedLinker.PackagePath, LoadFlags);
+			return PinObjects(Out);
 		}
 
 		return nullptr;
@@ -2048,7 +2073,7 @@ int32 RoundtripViaLinkerBatch(
 	TRACE_CPUPROFILER_EVENT_SCOPE(RoundtripViaLinkerBatch);
 
 	TArray<FName> PackageNames;
-	TArray<UPackage*> OriginalPackages;
+	TArray<UPackage*> SavedPackages;
 	TArray<UPackage*> BaselinePackages;
 	TArray<UPackage*> PlainPropsPackages;
 
@@ -2069,6 +2094,7 @@ int32 RoundtripViaLinkerBatch(
 		}
 
 		// Filter out in-memory packages
+		// and pin the remaining packages and their objects because UPackage::Save may trigger GC after bp compilation
 		for (UObject* PackageObj : AllPackages)
 		{
 			if (!PackageObj->HasAnyFlags(RF_WasLoaded))
@@ -2080,11 +2106,12 @@ int32 RoundtripViaLinkerBatch(
 			{
 				continue;
 			}
-			OriginalPackages.Add(Package);
-			PackageNames.Add(Package->GetFName());
+			SavedPackages.Add(Package);
+			PinPackage(Package);
+			PinObjects(Package);
 		}
 	}
-	const int32 NumPackages = OriginalPackages.Num();
+	int32 NumPackages = SavedPackages.Num();
 
 	FPlainPropsLinkerSandbox::SavedPackages.Reserve(NumPackages);
 	FPlainPropsLinkerSandbox::MountedBaselinePackages.Reserve(NumPackages);
@@ -2096,8 +2123,9 @@ int32 RoundtripViaLinkerBatch(
 		FScopeSaveTestPackage SaveScope;
 
 		TRACE_CPUPROFILER_EVENT_SCOPE(SavePackages);
-		for (UPackage* Package : OriginalPackages)
+		for (TArray<UPackage*>::TIterator It(SavedPackages); It; ++It)
 		{
+			UPackage* Package = *It;
 			UObject* Asset = Package->FindAssetInPackage();
 			FSavePackageArgs SaveArgs;
 			SaveArgs.TopLevelFlags = RF_Standalone;
@@ -2108,8 +2136,17 @@ int32 RoundtripViaLinkerBatch(
 				SaveArgs.SaveFlags |= SAVE_Unversioned_Properties;
 			}
 			FString SaveFilename = GetSavePackageFilename(Package);
-			UPackage::Save(Package, Asset, *SaveFilename, SaveArgs);
+			FSavePackageResultStruct Result = UPackage::Save(Package, Asset, *SaveFilename, SaveArgs);
+			if (Result.Result == ESavePackageResult::Success)
+			{
+				PackageNames.Add(Package->GetFName());
+			}
+			else
+			{
+				It.RemoveCurrent();
+			}
 		}
+		NumPackages = SavedPackages.Num();
 	}
 
 	// Roundtrip to text, also dumping all packages as yaml to disk
@@ -2133,9 +2170,9 @@ int32 RoundtripViaLinkerBatch(
 		TBitArray<> InstancedSet;
 		InstancedSet.Reserve(NumPackages);
 		{
-			TRACE_CPUPROFILER_EVENT_SCOPE(DropOriginalPackages);
-			UE_LOGFMT(LogPlainPropsUObject, Display, "Dropping {Num} original packages...", NumPackages);
-			DropPackages(CopyTemp(OriginalPackages), DropPackagesFunc);
+			TRACE_CPUPROFILER_EVENT_SCOPE(DropSavedPackages);
+			UE_LOGFMT(LogPlainPropsUObject, Display, "Dropping {Num} saved packages...", NumPackages);
+			DropPackages(CopyTemp(SavedPackages), DropPackagesFunc);
 			// Some packages may have been reloaded as a result of the drop package callbacks,
 			// mark that these should be loaded as instanced packages, i.e. with new unique in-memory package names
 			for (FName PackageName : PackageNames)
