PlainProps: More linker roundtrip test tweaks.
Unpin all packages and objects after the last UPackage::Save call since some DropPackage callbacks want to recreate a few selected objects and were asserting on the ref count.

#rnx

#Synced-CL 49808502

diff --git a/Engine/Plugins/Experimental/PlainPropsUObject/Source/Private/PlainPropsRoundtripTest.cpp b/Engine/Plugins/Experimental/PlainPropsUObject/Source/Private/PlainPropsRoundtripTest.cpp
--- a/Engine/Plugins/Experimental/PlainPropsUObject/Source/Private/PlainPropsRoundtripTest.cpp
+++ b/Engine/Plugins/Experimental/PlainPropsUObject/Source/Private/PlainPropsRoundtripTest.cpp
@@ -1307,17 +1307,23 @@ static UPackage* PinPackage(UPackage* Package)
 	if (Package)
 	{
 		Package->AddRef();
+		ForEachObjectWithPackage(Package, [](UObject* Obj)
+		{
+			Obj->AddRef();
+			return true;
+		});
 	}
 	return Package;
 }
 
-static UPackage* PinObjects(UPackage* Package)
+static UPackage* UnpinPackage(UPackage* Package)
 {
 	if (Package)
 	{
+		Package->ReleaseRef();
 		ForEachObjectWithPackage(Package, [](UObject* Obj)
 		{
-			Obj->AddRef();
+			Obj->ReleaseRef();
 			return true;
 		});
 	}
@@ -1460,7 +1466,7 @@ struct FPlainPropsLinkerSandbox
 			TUniquePtr<FMountedLinker> MountedInstance(new FMountedLinker(*SavedData, NewPackageName));
 			UPackage* Out = MountedInstance->Package;
 			MountedBaselinePackages.Emplace(NewPackageName, MoveTemp(MountedInstance));
-			return PinPackage(Out);
+			return Out;
 		}
 		return nullptr;
 	}
@@ -1478,7 +1484,7 @@ struct FPlainPropsLinkerSandbox
 			TUniquePtr<FMountedLinkerBatch> MountedInstance(new FMountedLinkerBatch(*SavedData, NewPackageName));
 			UPackage* Out = MountedInstance->MountedLinker.Package;
 			MountedPlainpropsPackages.Add(NewPackageName, MoveTemp(MountedInstance));
-			return PinPackage(Out);
+			return Out;
 		}
 		return nullptr;
 	}
@@ -1497,8 +1503,7 @@ struct FPlainPropsLinkerSandbox
 			FMountedLinker& MountedLinker = **MountedLinkerPtr;
 
 			UE_LOGFMT(LogPlainPropsUObject, Log, " Loading baseline for {Package}...", MountedLinker.PackageName);
-			UPackage* Out = LoadPackage(MountedLinker.Package, MountedLinker.PackagePath, LoadFlags);
-			return PinObjects(Out);
+			return LoadPackage(MountedLinker.Package, MountedLinker.PackagePath, LoadFlags);
 		}
 		return nullptr;
 	}
@@ -1510,8 +1515,7 @@ struct FPlainPropsLinkerSandbox
 			FMountedLinkerBatch& LinkerBatch = **LinkerBatchPtr;
 
 			UE_LOGFMT(LogPlainPropsUObject, Log, " Loading PlainProps for {Package}...", LinkerBatch.MountedLinker.PackageName);
-			UPackage* Out = LoadPackage(LinkerBatch.MountedLinker.Package, LinkerBatch.MountedLinker.PackagePath, LoadFlags);
-			return PinObjects(Out);
+			return LoadPackage(LinkerBatch.MountedLinker.Package, LinkerBatch.MountedLinker.PackagePath, LoadFlags);
 		}
 
 		return nullptr;
@@ -2108,7 +2112,6 @@ int32 RoundtripViaLinkerBatch(
 			}
 			SavedPackages.Add(Package);
 			PinPackage(Package);
-			PinObjects(Package);
 		}
 	}
 	int32 NumPackages = SavedPackages.Num();
@@ -2143,10 +2146,17 @@ int32 RoundtripViaLinkerBatch(
 			}
 			else
 			{
+				UnpinPackage(Package);
 				It.RemoveCurrent();
 			}
 		}
 		NumPackages = SavedPackages.Num();
+
+		// Unpin remaining packages, no more GC is expected, and some DropPackage callbacks tend to recreate uobjects.
+		for (UPackage* Package : SavedPackages)
+		{
+			UnpinPackage(Package);
+		}
 	}
 
 	// Roundtrip to text, also dumping all packages as yaml to disk
