PlainProps: Bug fixes for roundtripping an internal project.
Fix binding for Verse::FNativeString that has changed underlying storage from a direct FUtf8String to an FUtf8String wrapped in a TSharedPtr.
Fix use-after-free in GOptionals and GSets by cloning range bindings used in map keys.
Fix creation of FMemcpy plans with incorrect offsets for structs with unused supers.
Make sure to always patch up new plans in UpdateBatch (Memcpy only scenario).
Decorate the warning "Property '{Property}' is a '{Container}' of bitfield bools..." with owner information.

#rnx

#Synced-CL 49808502

diff --git a/Engine/Plugins/Experimental/PlainProps/Source/Private/PlainPropsLoad.cpp b/Engine/Plugins/Experimental/PlainProps/Source/Private/PlainPropsLoad.cpp
--- a/Engine/Plugins/Experimental/PlainProps/Source/Private/PlainPropsLoad.cpp
+++ b/Engine/Plugins/Experimental/PlainProps/Source/Private/PlainPropsLoad.cpp
@@ -472,10 +472,6 @@ struct FLoadPlanner
 				Mapping.LoadId = { Idx };
 				Mapping.Id = FBindId(RuntimeIds[Idx]);
 				CreatePlan(Mapping);
-				if (IsBound(Plans[Mapping.ReadId.Idx]))
-				{
-					UnboundSaveIds.Remove(Mapping.ReadId);
-				}
 			}
 		}
 
@@ -573,7 +569,7 @@ private:
 			check(OutReadId[-1].Idx < Batch.NumReadSchemas);
 		}
 
-		// Copy cloned subset schemas and patch up new plan pointers
+		// Copy cloned subset schemas and patch up new plans
 		if (SubsetSchemaData.Num() > 0)
 		{
 			// Allocate new subset data and add it to the linked list of dynamic allocations
@@ -604,6 +600,16 @@ private:
 			check(It == (uint8*)OutSubsetData + SubsetSchemaData.Num());
 			check(It == (uint8*)Ptr + Bytes);
 		}
+		else // simply patch up new plans
+		{
+			for (uint32 Idx = 0; Idx < NumPlans; ++Idx)
+			{
+				if (IsBound(Plans[Idx]))
+				{
+					Batch.Plans[Idx] = Plans[Idx];
+				}
+			}
+		}
 
 		Batch.NumPlans = NumPlans;
 	}
@@ -977,7 +983,7 @@ private:
 		}
 		
 		// Check all schema members are contiguous leaves
-		const uint32* OffsetIt = To.GetOffsets() + SkipToIdx;
+		const uint32* OffsetIt = To.GetOffsets() + SkipToIdx + To.HasSuper();
 		const uint32 StartPos = *OffsetIt;
 		uint32 EndPos = *OffsetIt;
 		for (FMemberType Member : From.GetMemberTypes())
@@ -985,7 +991,7 @@ private:
 			// Note: By adding FStructType::IsDense flag and some struct size lookup, memcpying of nested non-dynamic non-custom-bound structs might be possible
 			uint32 Offset = *OffsetIt++;
 			if (Offset != EndPos ||	// Non-contiguous or padded
-				Member.IsStruct() || // Nested dtructs have a skippable size prefix that can't be memcopied
+				Member.IsStruct() || // Nested structs have a skippable size prefix that can't be memcopied
 				Member.AsLeaf().Type == ELeafType::Bool) // bool values are stored as bits, see FBitCacheReader
 			{
 				return NullOpt;
diff --git a/Engine/Plugins/Experimental/PlainPropsUObject/Source/Private/PlainPropsUObjectBindings.cpp b/Engine/Plugins/Experimental/PlainPropsUObject/Source/Private/PlainPropsUObjectBindings.cpp
--- a/Engine/Plugins/Experimental/PlainPropsUObject/Source/Private/PlainPropsUObjectBindings.cpp
+++ b/Engine/Plugins/Experimental/PlainPropsUObject/Source/Private/PlainPropsUObjectBindings.cpp
@@ -1271,6 +1271,59 @@ static FRangeBinding AllocateArrayBinding(FArrayProperty* Property)
 	return FRangeBinding(*Out, ESizeType::S32);
 }
 
+class FRangeBindingAllocator
+{
+public:
+	TConstArrayView<FRangeBinding> Clone(FRangeBinding Head, TConstArrayView<FRangeBinding> Tail)
+	{
+		if (Tail.Num() == 0)
+		{
+			return MakeArrayView(&Ranges.Add_GetRef(Head), 1);
+		}
+		
+		// Ensure contiguous out range by padding up with dummy tail slice
+		static constexpr uint32 PageMax = decltype(Ranges)::MaxPerPage();
+		const int32 OutNum = 1 + Tail.Num();
+		check(OutNum <= PageMax);
+		const int32 NewPages = Align(Ranges.Num() + OutNum, PageMax) / PageMax;
+		if (NewPages > Ranges.NumPages() && !Ranges.IsEmpty())
+		{
+			int32 NumPad = Align(Ranges.Num(), PageMax) - Ranges.Num();
+			Ranges.Append(Tail.Slice(0, NumPad));
+			check(Ranges.Num() % PageMax == 0);
+		}
+
+		const FRangeBinding* OutData = &Ranges.Add_GetRef(Head);
+		Ranges.Append(Tail);
+		return MakeArrayView(OutData, OutNum);
+	}
+
+	TConstArrayView<FRangeBinding> Clone(TConstArrayView<FRangeBinding> Src)
+	{
+		if (Src.Num())
+		{
+			return Clone(Src[0], Src.RightChop(1));
+		}
+		return {};
+	}
+
+	void Reset()
+	{
+		Ranges.Reset();
+	}
+
+	void Empty()
+	{
+		Ranges.Empty();
+	}
+
+private:
+	TPagedArray<FRangeBinding, 1024> Ranges;
+};
+
+static FRangeBindingAllocator GScratchRanges;
+static FRangeBindingAllocator GRanges;
+
 ////////////////////////////////////////////////////////////////////////////////////////////////
 
 // Helpers to avoid using leaf FProperty instances after binding
@@ -2281,6 +2334,10 @@ public:
 		FDeclId DeclId = Param.IsLowered() ? GUE.Names.IndexDeclId(DeclType) : LowerCast(BindId);
 
 		DeltaBindNew(Property, Elem, {BindId, DeclId});
+		if (Elem.RangeBindings.Num())
+		{
+			Elem.RangeBindings = GRanges.Clone(Elem.RangeBindings);
+		}
 		return Bindings.Emplace(Elem, BindId);
 	}
 };
@@ -2990,6 +3047,10 @@ public:
 
 		// Todo: Ownership / memory leak
 		const IItemRangeBinding* Out = bIntrusive ? new FIntrusiveOptionalBinding(Inner) : BindNew(Inner);
+		if (Key.RangeBindings.Num())
+		{
+			Key.RangeBindings = GRanges.Clone(Key.RangeBindings);
+		}
 		return Bindings.Emplace(Key, FRangeBinding(*Out, ESizeType::Uni));
 	}
 };
@@ -2998,12 +3059,38 @@ static FOptionalBindings GOptionals;
 
 //////////////////////////////////////////////////////////////////////////////////////////////
 
+// Temporary unsafe helper for Verse::FNativeString that uses a FUtf8String wrapped in a TSharedPtr
+template<class StringType>
+struct TStringPtrBinding : TStringBinding<StringType>
+{
+	using TStringBinding<StringType>::TStringBinding;
+	
+	virtual void SaveLeaves(const void* Range, FLeafRangeAllocator& Out) const override
+	{
+		const void** Ptr = reinterpret_cast<const void**>(const_cast<void*>(Range));
+		TStringBinding<StringType>::SaveLeaves(*Ptr, Out);
+	}
+
+	virtual void LoadLeaves(void* Range, FLeafRangeLoadView Items) const override
+	{
+		void** Ptr = reinterpret_cast<void**>(Range);
+		TStringBinding<StringType>::LoadLeaves(*Ptr, Items);
+	}
+
+	virtual bool DiffLeaves(const void* RangeA, const void* RangeB) const override
+	{
+		const void** A = reinterpret_cast<const void**>(const_cast<void*>(RangeA));
+		const void** B = reinterpret_cast<const void**>(const_cast<void*>(RangeB));
+		return TStringBinding<StringType>::DiffLeaves(*A, *B);
+	}
+};
+
 struct FStringBindings
 {
 	TStringBinding<FString>			TCharInstance{GUE.Typenames.String};
 	TStringBinding<FUtf8String>		Utf8Instance{GUE.Typenames.Utf8String};
 	TStringBinding<FAnsiString>		AnsiInstance{GUE.Typenames.AnsiString};
-	TStringBinding<FUtf8String>		VerseInstance{GUE.Typenames.VerseString}; // Bypass Verse::FNativeString for now
+	TStringPtrBinding<FUtf8String>	VerseInstance{GUE.Typenames.VerseString}; // Bypass Verse::FNativeString for now, it uses a FUtf8String in wrapped in a TSharedPtr
 	FRangeBinding					TChar{TCharInstance, DefaultRangeMax};
 	FRangeBinding					Utf8{Utf8Instance, DefaultRangeMax};
 	FRangeBinding					Ansi{AnsiInstance, DefaultRangeMax};
@@ -3105,7 +3192,6 @@ private:
 	FOptionalDeclId									Super;
 	TOptional<FScopeId>								OwnerScope;
 	const EMemberPresence							Occupancy;
-	TPagedArray<FRangeBinding, 1024>				Ranges;
 	TArray<FMemberId, TInlineAllocator<64>>			Names;
 	TArray<FMemberBinding, TInlineAllocator<64>>	Members;
 	TArray<FMemberSpec, TInlineAllocator<64>>		Specs;
@@ -3115,30 +3201,6 @@ private:
 	const FName										VerseDynamicProperty{"VerseDynamicProperty"};
 	const FName										ReferenceProperty{"ReferenceProperty"}; // Verse reference + FProperty*
 
-	TConstArrayView<FRangeBinding> AllocateRangeBindings(FRangeBinding Head, TConstArrayView<FRangeBinding> Tail)
-	{
-		if (Tail.Num() == 0)
-		{
-			return MakeArrayView(&Ranges.Add_GetRef(Head), 1);
-		}
-		
-		// Ensure contiguous out range by padding up with dummy tail slice
-		static constexpr uint32 PageMax = decltype(Ranges)::MaxPerPage();
-		const int32 OutNum = 1 + Tail.Num();
-		check(OutNum <= PageMax);
-		const int32 NewPages = Align(Ranges.Num() + OutNum, PageMax) / PageMax;
-		if (NewPages > Ranges.NumPages() && !Ranges.IsEmpty())
-		{
-			int32 NumPad = Align(Ranges.Num(), PageMax) - Ranges.Num();
-			Ranges.Append(Tail.Slice(0, NumPad));
-			check(Ranges.Num() % PageMax == 0);
-		}
-
-		const FRangeBinding* OutData = &Ranges.Add_GetRef(Head);
-		Ranges.Append(Tail);
-		return MakeArrayView(OutData, OutNum);
-	}
-
 	static FMemberBinding Todo(FProperty* Property)
 	{
 		return FMemberBinding(Property->GetOffset_ForInternal());
@@ -3159,15 +3221,16 @@ private:
 	{
 		if (Inner.InnermostType.IsLeaf() && Inner.InnermostType.AsLeaf().Bind.Type == ELeafBindType::BitfieldBool)
 		{
-			UE_LOGFMT(LogPlainPropsUObject, Warning, "Property '{Property}' is a '{Container}' of bitfield bools, which make no sense. Binding as range of bools.", 
-				Property->GetFName(), GUE.Debug.Print(RangeBinding.GetBindName()));
+			UE_LOGFMT(LogPlainPropsUObject, Warning,
+				"Property '{Property}' in '{Owner}' is a '{Container}' of bitfield bools, which make no sense. Binding as range of bools.", 
+				Property->GetFName(), *Property->GetOwnerStruct()->GetPathName(), GUE.Debug.Print(RangeBinding.GetBindName()));
 			Inner.InnermostType = FMemberBindType(ReflectArithmetic<bool>);
 		}
 
 		FMemberBinding Out(Property->GetOffset_ForInternal());
 		Out.InnermostType = Inner.InnermostType;
 		Out.InnermostId = Inner.InnermostId;
-		Out.RangeBindings = AllocateRangeBindings(RangeBinding, Inner.RangeBindings);
+		Out.RangeBindings = GScratchRanges.Clone(RangeBinding, Inner.RangeBindings);
 		return Out;
 	}
 
@@ -3380,7 +3443,7 @@ private:
 			// Todo: Ownership / memory leak
 			const FStaticArrayBinding& ItemBinding = *new FStaticArrayBinding(ArrayDim, ElementSize);
 			ESizeType SizeType = ArrayDim < 256 ? ESizeType::U8 : ((ArrayDim < 65536) ? ESizeType::U16 : ESizeType::U32);
-			Out.RangeBindings =	AllocateRangeBindings(FRangeBinding(ItemBinding, SizeType), Out.RangeBindings);
+			Out.RangeBindings =	GScratchRanges.Clone(FRangeBinding(ItemBinding, SizeType), Out.RangeBindings);
 		}
 		else
 		{
@@ -3628,7 +3691,9 @@ void SchemaBindAllTypes(EBindMode Mode, EBatchType BatchType)
 		}
 
 		BindStruct(Id, Struct);
+		GScratchRanges.Reset();
 	}
+	GScratchRanges.Empty();
 }
 
 ////////////////////////////////////////////////////////////////////////////////////////////////
