PlainProps: Remove some duplicated PP_WITH_ENGINE_HOOKS.

#rnx

#Synced-CL 49987382

diff --git a/Engine/Plugins/Experimental/PlainPropsUObject/Source/Private/PlainPropsRoundtripTest.cpp b/Engine/Plugins/Experimental/PlainPropsUObject/Source/Private/PlainPropsRoundtripTest.cpp
--- a/Engine/Plugins/Experimental/PlainPropsUObject/Source/Private/PlainPropsRoundtripTest.cpp
+++ b/Engine/Plugins/Experimental/PlainPropsUObject/Source/Private/PlainPropsRoundtripTest.cpp
@@ -1167,17 +1167,13 @@ struct FLinkerPropertyBatch
 			if (PackageIdx.IsImport())
 			{
 				UObject* Object = LinkerLoad->ImportMap[PackageIdx.ToImport()].XObject;
-#if PP_WITH_ENGINE_HOOKS
 				Object = Object ? Object : LinkerLoad->ResolveResource(PackageIdx);
-#endif
 				Out = ::UE::CoreUObject::Private::MakeObjectHandle(Object);
 			}
 			else if (PackageIdx.IsExport())
 			{
 				UObject* Object = LinkerLoad->ExportMap[PackageIdx.ToExport()].Object;
-#if PP_WITH_ENGINE_HOOKS
 				Object = Object ? Object : LinkerLoad->ResolveResource(PackageIdx);
-#endif
 				Out = ::UE::CoreUObject::Private::MakeObjectHandle(Object);
 			}
 			else
