Quell logging when unmapped object are actually transient or trashed object

#Synced-CL 49808502

diff --git a/Engine/Plugins/Experimental/PlainPropsUObject/Source/Private/PlainPropsRoundtripTest.cpp b/Engine/Plugins/Experimental/PlainPropsUObject/Source/Private/PlainPropsRoundtripTest.cpp
--- a/Engine/Plugins/Experimental/PlainPropsUObject/Source/Private/PlainPropsRoundtripTest.cpp
+++ b/Engine/Plugins/Experimental/PlainPropsUObject/Source/Private/PlainPropsRoundtripTest.cpp
@@ -1151,8 +1151,11 @@ struct FLinkerPropertyBatch
 		int32 Idx = *reinterpret_cast<int32*>(&Index);
 		if (Object && Index.IsNull())
 		{
-			UE_LOGFMT(LogPlainPropsUObject, Warning, "UObject {Object} is not mapped.", Object->GetFullName());
-			Idx = 9999999;
+			if (!LinkerSave->IsUnsaveable(Object))
+			{
+				UE_LOGFMT(LogPlainPropsUObject, Warning, "UObject {Object} reference has not been harvested during save.", Object->GetFullName());
+			}
+			Idx = 9999999; // temp: make these objects easy to identify when looking at text output
 		}
 		Out.Add(GUE.Members.Id, Idx);
 	}
diff --git a/Engine/Source/Runtime/CoreUObject/Private/UObject/LinkerSave.cpp b/Engine/Source/Runtime/CoreUObject/Private/UObject/LinkerSave.cpp
--- a/Engine/Source/Runtime/CoreUObject/Private/UObject/LinkerSave.cpp
+++ b/Engine/Source/Runtime/CoreUObject/Private/UObject/LinkerSave.cpp
@@ -146,6 +146,15 @@ FLinkerSave::~FLinkerSave()
 	CloseAndDestroySaver();
 }
 
+bool FLinkerSave::IsUnsaveable(TObjectPtr<UObject> Object) const
+{
+	if (SaveContext)
+	{
+		return SaveContext->IsUnsaveable(Object, SaveContext->GetCachedObjectStatus(Object), /*bEmitWarning*/false);
+	}
+	return false;
+}
+
 int32 FLinkerSave::MapName(FNameEntryId Id) const
 {
 	const int32* IndexPtr = NameIndices.Find(Id);
diff --git a/Engine/Source/Runtime/CoreUObject/Public/UObject/LinkerSave.h b/Engine/Source/Runtime/CoreUObject/Public/UObject/LinkerSave.h
--- a/Engine/Source/Runtime/CoreUObject/Public/UObject/LinkerSave.h
+++ b/Engine/Source/Runtime/CoreUObject/Public/UObject/LinkerSave.h
@@ -186,6 +186,12 @@ public:
 		SaveContext = InContext;
 	}
 
+	/**
+	 * Query the temporary SaveContext for the saveable state of an object, if any.
+	 * Returns false if no save context is available
+	 */
+	UE_INTERNAL COREUOBJECT_API bool IsUnsaveable(TObjectPtr<UObject> Object) const;
+
 	/** Returns the appropriate name index for the source name, or INDEX_NONE if not found in NameIndices */
 	int32 MapName( FNameEntryId Name) const;
 
