PlainProps: Minor tweaks to the linker mapping of FNames and UObjects.

#rnx

#Synced-CL 49987382

diff --git a/Engine/Plugins/Experimental/PlainPropsUObject/Source/Private/PlainPropsRoundtripTest.cpp b/Engine/Plugins/Experimental/PlainPropsUObject/Source/Private/PlainPropsRoundtripTest.cpp
--- a/Engine/Plugins/Experimental/PlainPropsUObject/Source/Private/PlainPropsRoundtripTest.cpp
+++ b/Engine/Plugins/Experimental/PlainPropsUObject/Source/Private/PlainPropsRoundtripTest.cpp
@@ -1094,12 +1094,12 @@ struct FLinkerPropertyBatch
 		uint64 Id;
 		if (Index == INDEX_NONE)
 		{
-			UE_LOGFMT(LogPlainPropsUObject, Warning, "FName {Name} is not mapped.", In);
-			Id = 6666666;
+			UE_LOGFMT(LogPlainPropsUObject, Warning, "FName '{Name}' has not been harvested during save.", In);
+			Id = 6666666; // temp: make these names easy to identify when looking at text output
 		}
 		else
 		{
-			Id = Index != INDEX_NONE ? ((uint64(Number) << 32) | Index) : 6666666;
+			Id = ((uint64(Number) << 32) | Index);
 		}
 		Out.Add(GUE.Members.Id, Id);
 	}
@@ -1108,7 +1108,7 @@ struct FLinkerPropertyBatch
 	{
 		uint64 Id;
 		LoadSole<uint64>(&Id, In);
-		if (Id != 6666666)
+		if (Id != 6666666) // todo: 6666666 is prototype code
 		{
 			int32 Index = Id & 0xFFFFFFFF;
 			int32 Number = Id >> 32;
@@ -1153,9 +1153,9 @@ struct FLinkerPropertyBatch
 		{
 			if (!LinkerSave->IsUnsaveable(Object))
 			{
-				UE_LOGFMT(LogPlainPropsUObject, Warning, "UObject {Object} reference has not been harvested during save.", Object->GetFullName());
+				UE_LOGFMT(LogPlainPropsUObject, Warning, "UObject '{Object}' reference has not been harvested during save.", Object->GetFullName());
+				Idx = 9999999; // temp: make these objects easy to identify when looking at text output
 			}
-			Idx = 9999999; // temp: make these objects easy to identify when looking at text output
 		}
 		Out.Add(GUE.Members.Id, Idx);
 	}
@@ -1164,29 +1164,23 @@ struct FLinkerPropertyBatch
 	{
 		int32 Idx;
 		LoadSole<int32>(&Idx, In);
-		if (Idx != 9999999)
+		Idx = (Idx == 9999999) ? 0 : Idx; // todo: 9999999 is prototype code
+		FPackageIndex PackageIdx = *reinterpret_cast<FPackageIndex*>(&Idx);
+		if (PackageIdx.IsImport())
 		{
-			FPackageIndex PackageIdx = *reinterpret_cast<FPackageIndex*>(&Idx);
-			if (PackageIdx.IsImport())
-			{
-				UObject* Object = LinkerLoad->ImportMap[PackageIdx.ToImport()].XObject;
-				Object = Object ? Object : LinkerLoad->ResolveResource(PackageIdx);
-				Out = ::UE::CoreUObject::Private::MakeObjectHandle(Object);
-			}
-			else if (PackageIdx.IsExport())
-			{
-				UObject* Object = LinkerLoad->ExportMap[PackageIdx.ToExport()].Object;
-				Object = Object ? Object : LinkerLoad->ResolveResource(PackageIdx);
-				Out = ::UE::CoreUObject::Private::MakeObjectHandle(Object);
-			}
-			else
-			{
-				check(PackageIdx.IsNull());
-				Out = FObjectHandle();
-			}
+			UObject* Object = LinkerLoad->ImportMap[PackageIdx.ToImport()].XObject;
+			Object = Object ? Object : LinkerLoad->ResolveResource(PackageIdx);
+			Out = ::UE::CoreUObject::Private::MakeObjectHandle(Object);
+		}
+		else if (PackageIdx.IsExport())
+		{
+			UObject* Object = LinkerLoad->ExportMap[PackageIdx.ToExport()].Object;
+			Object = Object ? Object : LinkerLoad->ResolveResource(PackageIdx);
+			Out = ::UE::CoreUObject::Private::MakeObjectHandle(Object);
 		}
 		else
 		{
+			check(PackageIdx.IsNull());
 			Out = FObjectHandle();
 		}
 	}
