PlainProps: Fix broken binding for VerseString.

#rnx

#Synced-CL 49987382

diff --git a/Engine/Plugins/Experimental/PlainPropsUObject/Source/Private/PlainPropsUObjectBindings.cpp b/Engine/Plugins/Experimental/PlainPropsUObject/Source/Private/PlainPropsUObjectBindings.cpp
--- a/Engine/Plugins/Experimental/PlainPropsUObject/Source/Private/PlainPropsUObjectBindings.cpp
+++ b/Engine/Plugins/Experimental/PlainPropsUObject/Source/Private/PlainPropsUObjectBindings.cpp
@@ -14,6 +14,7 @@
 #include "Math/PreciseFP.h"
 #include "Misc/DefinePrivateMemberPtr.h"
 #include "StructUtils/UserDefinedStruct.h"
+#include "Templates/SharedPointer.h"
 #include "Templates/UniquePtr.h"
 #include "UObject/AnsiStrProperty.h"
 #include "UObject/Class.h"
@@ -30,6 +31,7 @@
 #include "UObject/Utf8StrProperty.h"
 #include "UObject/VerseStringProperty.h"
 #include "UObject/VerseValueProperty.h"
+#include "VerseVM/VVMNativeString.h"
 
 
 using FUnicastScriptDelegate = TScriptDelegate<FNotThreadSafeNotCheckedDelegateMode>;
@@ -3060,28 +3062,28 @@ static FOptionalBindings GOptionals;
 //////////////////////////////////////////////////////////////////////////////////////////////
 
 // Temporary unsafe helper for Verse::FNativeString that uses a FUtf8String wrapped in a TSharedPtr
-template<class StringType>
-struct TStringPtrBinding : TStringBinding<StringType>
+struct FVerseStringBinding : TStringBinding<FUtf8String>
 {
-	using TStringBinding<StringType>::TStringBinding;
+	using TStringBinding<FUtf8String>::TStringBinding;
 	
 	virtual void SaveLeaves(const void* Range, FLeafRangeAllocator& Out) const override
 	{
 		const void** Ptr = reinterpret_cast<const void**>(const_cast<void*>(Range));
-		TStringBinding<StringType>::SaveLeaves(*Ptr, Out);
+		TStringBinding<FUtf8String>::SaveLeaves(*Ptr, Out);
 	}
 
 	virtual void LoadLeaves(void* Range, FLeafRangeLoadView Items) const override
 	{
-		void** Ptr = reinterpret_cast<void**>(Range);
-		TStringBinding<StringType>::LoadLeaves(*Ptr, Items);
+		FUtf8String Str;
+		TStringBinding<FUtf8String>::LoadLeaves(&Str, Items);
+		*static_cast<Verse::FNativeString*>(Range) = MoveTemp(Str);
 	}
 
 	virtual bool DiffLeaves(const void* RangeA, const void* RangeB) const override
 	{
 		const void** A = reinterpret_cast<const void**>(const_cast<void*>(RangeA));
 		const void** B = reinterpret_cast<const void**>(const_cast<void*>(RangeB));
-		return TStringBinding<StringType>::DiffLeaves(*A, *B);
+		return TStringBinding<FUtf8String>::DiffLeaves(*A, *B);
 	}
 };
 
@@ -3090,7 +3092,7 @@ struct FStringBindings
 	TStringBinding<FString>			TCharInstance{GUE.Typenames.String};
 	TStringBinding<FUtf8String>		Utf8Instance{GUE.Typenames.Utf8String};
 	TStringBinding<FAnsiString>		AnsiInstance{GUE.Typenames.AnsiString};
-	TStringPtrBinding<FUtf8String>	VerseInstance{GUE.Typenames.VerseString}; // Bypass Verse::FNativeString for now, it uses a FUtf8String in wrapped in a TSharedPtr
+	FVerseStringBinding				VerseInstance{GUE.Typenames.VerseString};
 	FRangeBinding					TChar{TCharInstance, DefaultRangeMax};
 	FRangeBinding					Utf8{Utf8Instance, DefaultRangeMax};
 	FRangeBinding					Ansi{AnsiInstance, DefaultRangeMax};
