nss: JSON support for roundtrip commandlet.

#rnx

#Synced-CL 50479750

diff --git a/Engine/Plugins/Experimental/PlainPropsEngine/Source/Private/PlainPropsCommandlets.cpp b/Engine/Plugins/Experimental/PlainPropsEngine/Source/Private/PlainPropsCommandlets.cpp
--- a/Engine/Plugins/Experimental/PlainPropsEngine/Source/Private/PlainPropsCommandlets.cpp
+++ b/Engine/Plugins/Experimental/PlainPropsEngine/Source/Private/PlainPropsCommandlets.cpp
@@ -335,6 +335,11 @@ int32 UTestPlainPropsCommandlet::Main(const FString& Params)
 		Options = ERoundtrip::TextMemory | ERoundtrip::TextStable;
 	}
 
+	if (Params.Find("-json") != INDEX_NONE)
+	{
+		Options |= ERoundtrip::JSON;
+	}
+
 	UE_LOGFMT(LogPlainPropsEngine, Display, "Starting roundtrip test...");
 	
 	// Roundtrip
diff --git a/Engine/Plugins/Experimental/PlainPropsUObject/Source/Private/PlainPropsRoundtripTest.cpp b/Engine/Plugins/Experimental/PlainPropsUObject/Source/Private/PlainPropsRoundtripTest.cpp
--- a/Engine/Plugins/Experimental/PlainPropsUObject/Source/Private/PlainPropsRoundtripTest.cpp
+++ b/Engine/Plugins/Experimental/PlainPropsUObject/Source/Private/PlainPropsRoundtripTest.cpp
@@ -829,6 +829,8 @@ FORCENOINLINE void LoadArchive(FPropertyReader& Archive, TConstArrayView<FInstan
 
 ////////////////////////////////////////////////////////////////////////////////////////////////
 
+enum class EMarkupFormat { Yaml, Json };
+
 class FStableNameBatchIds final : public FStableBatchIds
 {
 	TConstArrayView<FSensitiveName> Names;
@@ -843,10 +845,15 @@ public:
 [[nodiscard]] static FSchemaBatchId ParseBatch(
 	TArray64<uint8>& OutData,
 	TArray<FStructView>& OutObjects,
-	FUtf8StringView YamlView)
+	FUtf8StringView MarkupView,
+	EMarkupFormat MarkupFormat=EMarkupFormat::Yaml)
 {
-	// Parse yaml
-	ParseYamlBatch(OutData, YamlView);
+	// Parse markup
+	switch (MarkupFormat)
+	{
+	case EMarkupFormat::Yaml: ParseYamlBatch(OutData, MarkupView); break;
+	case EMarkupFormat::Json: ParseJsonBatch(OutData, MarkupView); break;
+	}
 
 	// Grab and mount parsed schemas
 	FByteReader It(MakeMemoryView(OutData));
@@ -871,32 +878,38 @@ static void RoundtripText(
 	TConstArrayView<FStructView> Objects,
 	TConstArrayView<FInstance> Instances,
 	ESchemaFormat Format,
-	FStringView AssetName = {})
+	FStringView AssetName = {},
+	EMarkupFormat MarkupFormat = EMarkupFormat::Yaml)
 {
-	// Print yaml
+	// Print markup
 	UE_LOGFMT(LogPlainPropsUObject, Display, "Printing to PlainProps text using {Format}...", ToString(Format));
-	TUtf8StringBuilder<256> Yaml;
-	Yaml.Reserve(INT_MAX);
-	PrintYamlBatch(Yaml, BatchIds, Objects);
-	FUtf8StringView YamlView = Yaml.ToView();
+	TUtf8StringBuilder<256> Markup;
+	Markup.Reserve(INT_MAX);
+	switch (MarkupFormat)
+	{
+	case EMarkupFormat::Yaml: PrintYamlBatch(Markup, BatchIds, Objects); break;
+	case EMarkupFormat::Json: PrintJsonBatch(Markup, BatchIds, Objects); break;
+	}
+	FUtf8StringView MarkupView = Markup.ToView();
 
 	// Write to file
 	if (AssetName.Len() == 0)
 	{
 		AssetName = (Format == ESchemaFormat::InMemoryNames ? TEXTVIEW("DiffTest_InMemoryNames") : TEXTVIEW("DiffTest_StableNames"));
 	}
-	const FString YamlPath = FPaths::Combine(FPaths::ProjectSavedDir(), TEXT("PlainProps"), AssetName) + TEXT(".yaml");
-	if (TUniquePtr<FArchive> FileWriter(IFileManager::Get().CreateFileWriter(*YamlPath)); FileWriter)
+	const TCHAR* Extension = (MarkupFormat == EMarkupFormat::Yaml) ? TEXT(".yaml") : TEXT(".json");
+	const FString MarkupPath = FPaths::Combine(FPaths::ProjectSavedDir(), TEXT("PlainProps"), AssetName) + Extension;
+	if (TUniquePtr<FArchive> FileWriter(IFileManager::Get().CreateFileWriter(*MarkupPath)); FileWriter)
 	{
-		UE_LOGFMT(LogPlainPropsUObject, Display, "Writing {KB}KB yaml as {YamlPath}...", Yaml.Len() >> 10, *YamlPath);
-		FileWriter->Serialize((void*)YamlView.GetData(), YamlView.Len());
+		UE_LOGFMT(LogPlainPropsUObject, Display, "Writing {KB}KB {Format} as {MarkupPath}...", Markup.Len() >> 10, Extension + 1, *MarkupPath);
+		FileWriter->Serialize((void*)MarkupView.GetData(), MarkupView.Len());
 	}
 
-	// Parse yaml
+	// Parse text markup
 	UE_LOGFMT(LogPlainPropsUObject, Display, "Parsing PlainProps text using {Format}...", ToString(Format));
 	TArray64<uint8> Data;
 	TArray<FStructView> ParsedObjects;
-	FSchemaBatchId ParsedBatch = ParseBatch(Data, ParsedObjects, YamlView);
+	FSchemaBatchId ParsedBatch = ParseBatch(Data, ParsedObjects, MarkupView, MarkupFormat);
 
 	if (Format == ESchemaFormat::StableNames && Instances.Num())
 	{
@@ -926,8 +939,8 @@ static void RoundtripText(
 			}
 		}
 		UE_LOGFMT(LogPlainPropsUObject, Display,
-			"Detected {Diffs} diffs in {Objs} PlainProps parsed objects from {KB}KB yaml text using StableNames\n{Diffs}",
-			NumDiffs, NumObjects, Yaml.Len() >> 10, Diffs.ToString());
+			"Detected {Diffs} diffs in {Objs} PlainProps parsed objects from {KB}KB {Format} text using StableNames\n{Diffs}",
+			NumDiffs, NumObjects, Markup.Len() >> 10, Extension + 1, Diffs.ToString());
 	}
 
 	// Unmount parsed schemas
@@ -937,7 +950,9 @@ static void RoundtripText(
 class FBatchTextRoundtripper
 {
 public:
-	FBatchTextRoundtripper(FMemoryView Data, ESchemaFormat InFormat) : Format(InFormat)
+	FBatchTextRoundtripper(FMemoryView Data, ESchemaFormat InFormat, EMarkupFormat InMarkupFormat)
+	: Format(InFormat)
+	, MarkupFormat(InMarkupFormat)
 	{
 		FByteReader It(Data);
 		verify(It.Grab<uint32>() == Magics[0]);
@@ -1001,27 +1016,32 @@ public:
 
 	void RoundtripTextForDiffTest(TConstArrayView<FInstance> Instances) const
 	{
-		PlainProps::UE::RoundtripText(*BatchIds, Objects, Instances, Format, {});
+		PlainProps::UE::RoundtripText(*BatchIds, Objects, Instances, Format, {}, MarkupFormat);
 	}
 
 	void RoundtripTextForPackage(FStringView PackageName) const
 	{
-		PlainProps::UE::RoundtripText(*BatchIds, Objects, {}, Format, PackageName);
+		PlainProps::UE::RoundtripText(*BatchIds, Objects, {}, Format, PackageName, MarkupFormat);
 	}
 
-	void PrintPackageToYamlFile(FStringView PackageName) const
+	void PrintPackageToFile(FStringView PackageName) const
 	{
-		TUtf8StringBuilder<64> Yaml;
-		Yaml.Reserve(2 << 10);
-		PrintYamlBatch(Yaml, *BatchIds, Objects);
-		FUtf8StringView YamlView = Yaml.ToView();
+		TUtf8StringBuilder<64> Markup;
+		Markup.Reserve(2 << 10);
+		switch (MarkupFormat)
+		{
+		case EMarkupFormat::Yaml: PrintYamlBatch(Markup, *BatchIds, Objects); break;
+		case EMarkupFormat::Json: PrintJsonBatch(Markup, *BatchIds, Objects); break;
+		}
+		FUtf8StringView MarkupView = Markup.ToView();
 
 		// Write to file
-		const FString YamlPath = FPaths::Combine(FPaths::ProjectSavedDir(), TEXT("PlainProps"), PackageName) + TEXT(".yaml");
-		if (TUniquePtr<FArchive> Ar(IFileManager::Get().CreateFileWriter(*YamlPath)); Ar)
+		const TCHAR* Extension = (MarkupFormat == EMarkupFormat::Yaml) ? TEXT(".yaml") : TEXT(".json");
+		const FString Path = FPaths::Combine(FPaths::ProjectSavedDir(), TEXT("PlainProps"), PackageName) + Extension;
+		if (TUniquePtr<FArchive> Ar(IFileManager::Get().CreateFileWriter(*Path)); Ar)
 		{
-			UE_LOGFMT(LogPlainPropsUObject, Log, "Writing {KB}KB yaml as {YamlPath}...", Yaml.Len() >> 10, *YamlPath);
-			Ar->Serialize((void*)YamlView.GetData(), YamlView.Len());
+			UE_LOGFMT(LogPlainPropsUObject, Log, "Writing {KB}KB {Format} as {Path}...", Markup.Len() >> 10, Extension + 1, *Path);
+			Ar->Serialize((void*)MarkupView.GetData(), MarkupView.Len());
 		}
 	}
 
@@ -1030,6 +1050,7 @@ private:
 	TArray<FSensitiveName>		Names;
 	TUniquePtr<FBatchIds>		BatchIds;
 	ESchemaFormat				Format;
+	EMarkupFormat				MarkupFormat;
 	
 	template<typename T>
 	static TConstArrayView<T> GrabNumAndArray(/* in-out */ FByteReader& It)
@@ -1514,21 +1535,21 @@ struct FPlainPropsLinkerSandbox
 		return nullptr;
 	}
 
-	static void RoundtripToTextFromMemory(FName PackageName)
+	static void RoundtripToTextFromMemory(FName PackageName, EMarkupFormat MarkupFormat)
 	{
 		if (FSavedPackage* SavedData = SavedPackages.Find(PackageName))
 		{
-			FBatchTextRoundtripper StableBatch(MakeMemoryView(SavedData->PlainPropsData), ESchemaFormat::StableNames);
+			FBatchTextRoundtripper StableBatch(MakeMemoryView(SavedData->PlainPropsData), ESchemaFormat::StableNames, MarkupFormat);
 			StableBatch.RoundtripTextForPackage(PackageName.ToString());
 		}
 	}
 
-	static void PrintPackageToYamlFile(FName PackageName)
+	static void PrintPackageToFile(FName PackageName, EMarkupFormat MarkupFormat)
 	{
 		if (FSavedPackage* SavedData = SavedPackages.Find(PackageName))
 		{
-			FBatchTextRoundtripper StableBatch(MakeMemoryView(SavedData->PlainPropsData), ESchemaFormat::StableNames);
-			StableBatch.PrintPackageToYamlFile(PackageName.ToString());
+			FBatchTextRoundtripper StableBatch(MakeMemoryView(SavedData->PlainPropsData), ESchemaFormat::StableNames, MarkupFormat);
+			StableBatch.PrintPackageToFile(PackageName.ToString());
 		}
 	}
 
@@ -1837,7 +1858,8 @@ int32 RoundtripViaPlainBatch(TConstArrayView<UObject*> Objects, ERoundtrip Optio
 
 		if (EnumHasAnyFlags(Options, ERoundtrip::TextMemory))
 		{
-			FBatchTextRoundtripper MemoryBatch(MakeMemoryView(Plain.Data), ESchemaFormat::InMemoryNames);
+			EMarkupFormat MarkupFormat = EnumHasAnyFlags(Options, ERoundtrip::JSON) ? EMarkupFormat::Json : EMarkupFormat::Yaml;
+			FBatchTextRoundtripper MemoryBatch(MakeMemoryView(Plain.Data), ESchemaFormat::InMemoryNames, MarkupFormat);
 			MemoryBatch.RoundtripTextForDiffTest(Instances);
 		}
 		if (EnumHasAnyFlags(Options, ERoundtrip::TextStable))
@@ -1845,7 +1867,8 @@ int32 RoundtripViaPlainBatch(TConstArrayView<UObject*> Objects, ERoundtrip Optio
 			UE_LOGFMT(LogPlainPropsUObject, Display, "Saving UObjects to PlainProps with StableNames...");
 			TArray64<uint8> StableData = Batch.Write();
 
-			FBatchTextRoundtripper StableBatch(MakeMemoryView(StableData), ESchemaFormat::StableNames);
+			EMarkupFormat MarkupFormat = EnumHasAnyFlags(Options, ERoundtrip::JSON) ? EMarkupFormat::Json : EMarkupFormat::Yaml;
+			FBatchTextRoundtripper StableBatch(MakeMemoryView(StableData), ESchemaFormat::StableNames, MarkupFormat);
 			StableBatch.RoundtripTextForDiffTest(Instances);
 		}
 	}
@@ -2155,11 +2178,12 @@ int32 RoundtripViaLinkerBatch(
 	// Roundtrip to text, also dumping all packages as yaml to disk
 	if (EnumHasAnyFlags(Options, ERoundtrip::TextMemory | ERoundtrip::TextStable))
 	{
+		EMarkupFormat MarkupFormat = EnumHasAnyFlags(Options, ERoundtrip::JSON) ? EMarkupFormat::Json : EMarkupFormat::Yaml;
 		TRACE_CPUPROFILER_EVENT_SCOPE(RoundtripToText);
 		UE_LOGFMT(LogPlainPropsUObject, Display, "Roundtripping {SaveNum} packages to text...", NumPackages);
 		for (FName PackageName : PackageNames)
 		{
-			FPlainPropsLinkerSandbox::RoundtripToTextFromMemory(PackageName);
+			FPlainPropsLinkerSandbox::RoundtripToTextFromMemory(PackageName, MarkupFormat);
 		}
 	}
 
@@ -2258,7 +2282,7 @@ int32 RoundtripViaLinkerBatch(
 			std::atomic<uint32> NumDiffs = 0;
 			std::atomic<uint32> NumIgnored = 0;
 			ParallelFor(TEXT("DiffPackage"), NumPackages, /*MinBatchSize*/1,
-				[&NumDiffs, &NumIgnored, &PackageNames, &BaselinePackages, &PlainPropsPackages, &Filter](int32 I)
+				[&NumDiffs, &NumIgnored, &PackageNames, &BaselinePackages, &PlainPropsPackages, &Filter, Options](int32 I)
 			{
 				FDiffObjectContext Ctx = { BaselinePackages[I], PlainPropsPackages[I], Filter };
 				{
@@ -2275,8 +2299,9 @@ int32 RoundtripViaLinkerBatch(
 					++NumDiffs;
 
 					{
-						TRACE_CPUPROFILER_EVENT_SCOPE(PrintToYaml);
-						FPlainPropsLinkerSandbox::PrintPackageToYamlFile(PackageNames[I]);
+						TRACE_CPUPROFILER_EVENT_SCOPE(PrintToMarkup);
+						EMarkupFormat MarkupFormat = EnumHasAnyFlags(Options, ERoundtrip::JSON) ? EMarkupFormat::Json : EMarkupFormat::Yaml;
+						FPlainPropsLinkerSandbox::PrintPackageToFile(PackageNames[I], MarkupFormat);
 					}
 
 					{
diff --git a/Engine/Plugins/Experimental/PlainPropsUObject/Source/Public/PlainPropsRoundtripTest.h b/Engine/Plugins/Experimental/PlainPropsUObject/Source/Public/PlainPropsRoundtripTest.h
--- a/Engine/Plugins/Experimental/PlainPropsUObject/Source/Public/PlainPropsRoundtripTest.h
+++ b/Engine/Plugins/Experimental/PlainPropsUObject/Source/Public/PlainPropsRoundtripTest.h
@@ -26,6 +26,7 @@ enum class ERoundtrip : uint8
 	UPS 		= 1 << 2,
 	TextMemory	= 1 << 3,
 	TextStable	= 1 << 4,
+	JSON		= 1 << 5,
 };
 ENUM_CLASS_FLAGS(ERoundtrip);
 
