PlainProps: Fix implicit conversion to bool.

#rnx

#Synced-CL 50970586

#Synced-CL 50970586

diff --git a/Engine/Plugins/Experimental/PlainPropsUObject/Source/Private/PlainPropsUObjectBindings.cpp b/Engine/Plugins/Experimental/PlainPropsUObject/Source/Private/PlainPropsUObjectBindings.cpp
--- a/Engine/Plugins/Experimental/PlainPropsUObject/Source/Private/PlainPropsUObjectBindings.cpp
+++ b/Engine/Plugins/Experimental/PlainPropsUObject/Source/Private/PlainPropsUObjectBindings.cpp
@@ -3533,7 +3533,7 @@ bool ShouldBind(const UStruct* Struct)
 	}
 	if (Struct->HasAnyCastFlags(CASTCLASS_UScriptStruct))
 	{
-		return static_cast<const UScriptStruct*>(Struct)->StructFlags & STRUCT_SerializeNative;
+		return !!(static_cast<const UScriptStruct*>(Struct)->StructFlags & STRUCT_SerializeNative);
 	}
 	return false;
 }
