First version of overridable serialization bindings for the new serialization stack.

#rnx
#Synced-CL 51840270

diff --git a/Engine/Plugins/Experimental/PlainProps/Source/Private/PlainPropsInternalBind.h b/Engine/Plugins/Experimental/PlainProps/Source/Private/PlainPropsInternalBind.h
--- a/Engine/Plugins/Experimental/PlainProps/Source/Private/PlainPropsInternalBind.h
+++ b/Engine/Plugins/Experimental/PlainProps/Source/Private/PlainPropsInternalBind.h
@@ -3,48 +3,11 @@
 #pragma once
 
 #include "PlainPropsBind.h"
+#include "PlainPropsVisitMember.h"
 
 namespace PlainProps 
 {
 
-// Iterates over member bindings
-class FMemberVisitor
-{
-public:
-	explicit FMemberVisitor(const FSchemaBinding& InSchema);
-
-	bool						HasMore() const			{ return MemberIdx < NumMembers; }
-	uint16						GetIndex() const		{ return MemberIdx; }
-	
-	EMemberKind					PeekKind() const;		// @pre HasMore()
-	FMemberBindType				PeekType() const;		// @pre HasMore()
-	uint32						PeekOffset() const;		// @pre HasMore()
-
-	FLeafMemberBinding			GrabLeaf();				// @pre PeekKind() == EMemberKind::Leaf
-	FRangeMemberBinding			GrabRange();			// @pre PeekKind() == EMemberKind::Range
-	FStructMemberBinding		GrabStruct();			// @pre PeekKind() == EMemberKind::Struct
-	FBindId						GrabSuper();			// @pre First grab and has declared super
-	void						SkipMember();
-
-protected: // for unit tests
-	const FSchemaBinding&		Schema;
-	const uint16				NumMembers;
-	uint16						MemberIdx = 0;
-	uint16						InnerRangeIdx = 0;		// Types of [nested] ranges
-	uint16						InnerIdIdx = 0;			// Types of static structs and enums
-
-	using FMemberBindTypeRange = TConstArrayView<FMemberBindType>;
-
-	uint64						GrabMemberOffset();
-	FMemberBindTypeRange		GrabInnerTypes();
-	FInnerId					GrabInnerId();
-	FBindId						GrabStructSchema(FStructType Type);
-	FOptionalInnerId			GrabRangeSchema(FMemberType InnermostType);
-	FEnumId						GrabEnumSchema()		{ return GrabInnerId().AsEnum(); }
-};
-
-////////////////////////////////////////////////////////////////////////////////////////////////
-
 struct FMemberBinderBase
 {
 	FMemberBinderBase(FSchemaBinding& InSchema)
diff --git a/Engine/Plugins/Experimental/PlainProps/Source/Private/PlainPropsRead.cpp b/Engine/Plugins/Experimental/PlainProps/Source/Private/PlainPropsRead.cpp
--- a/Engine/Plugins/Experimental/PlainProps/Source/Private/PlainPropsRead.cpp
+++ b/Engine/Plugins/Experimental/PlainProps/Source/Private/PlainPropsRead.cpp
@@ -357,6 +357,15 @@ FRangeView FMemberReader::GrabRange()
 //	return {Type, Id, BatchId, Value};
 //}
 
+void FMemberReader::SkipMember()
+{
+	switch (PeekKind())
+	{
+		case EMemberKind::Leaf:		(void)GrabLeaf();		break;
+		case EMemberKind::Struct:	(void)GrabStruct();		break;
+		case EMemberKind::Range:	(void)GrabRange();		break;
+	}
+}
 
 void FMemberReader::GrabLeaves(void* Out, uint32 Num, SIZE_T Size)
 {
diff --git a/Engine/Plugins/Experimental/PlainProps/Source/Private/PlainPropsSave.cpp b/Engine/Plugins/Experimental/PlainProps/Source/Private/PlainPropsSave.cpp
--- a/Engine/Plugins/Experimental/PlainProps/Source/Private/PlainPropsSave.cpp
+++ b/Engine/Plugins/Experimental/PlainProps/Source/Private/PlainPropsSave.cpp
@@ -293,24 +293,24 @@ inline const FStructDeclaration& GetDeclaration(const FSaveContext& Ctx, FBindId
 	return Out;
 }
 
-[[nodiscard]] static FMemberSchema CreateRangeSchema(FScratchAllocator& Scratch, FRangeMemberBinding Member)
+[[nodiscard]] FMemberSchema CreateRangeSchema(FScratchAllocator& Scratch, FRangeMemberBinding Member)
 {
 	FMemberType* InnerRangeTypes = CreateInnerRangeTypes(Scratch, Member.NumRanges, Member.InnerTypes);
 	return { FMemberType(Member.RangeBindings[0].GetSizeType()), ToMemberType(Member.InnerTypes[0]), Member.NumRanges, Member.InnermostId, InnerRangeTypes };
 }
 
-static void SaveMember(FMemberBuilder& Out, const void* Struct, FMemberId Name, const FSaveContext& Ctx, FLeafMemberBinding Member)
+void SaveMember(FMemberBuilder& Out, const void* Struct, FMemberId Name, const FSaveContext& Ctx, FLeafMemberBinding Member)
 {
 	FUnpackedLeafType Type = ToUnpackedLeafType(Member.Leaf);
 	Out.AddLeaf(Name, Type, Member.Enum, SaveLeaf(At(Struct, Member.Offset), Member.Leaf));
 }
 
-static void SaveMember(FMemberBuilder& Out, const void* Struct, FMemberId Name, const FSaveContext& Ctx, FRangeMemberBinding Member)
+void SaveMember(FMemberBuilder& Out, const void* Struct, FMemberId Name, const FSaveContext& Ctx, FRangeMemberBinding Member)
 {
 	Out.AddRange(Name, { CreateRangeSchema(Ctx.Scratch, Member), SaveRange(At(Struct, Member.Offset), Member, Ctx) });
 }
 
-static void SaveMember(FMemberBuilder& Out, const void* Struct, FMemberId Name, const FSaveContext& Ctx, FStructMemberBinding Member)
+void SaveMember(FMemberBuilder& Out, const void* Struct, FMemberId Name, const FSaveContext& Ctx, FStructMemberBinding Member)
 {
 	Out.AddStruct(Name, Member.Id, SaveStruct(At(Struct, Member.Offset), Member.Id, Ctx));
 }
diff --git a/Engine/Plugins/Experimental/PlainProps/Source/Public/PlainPropsRead.h b/Engine/Plugins/Experimental/PlainProps/Source/Public/PlainPropsRead.h
--- a/Engine/Plugins/Experimental/PlainProps/Source/Public/PlainPropsRead.h
+++ b/Engine/Plugins/Experimental/PlainProps/Source/Public/PlainPropsRead.h
@@ -627,7 +627,7 @@ class FMemberReader
 {
 public:
 	explicit FMemberReader(FStructView Struct) : FMemberReader(Struct.Schema.Resolve(), Struct.Values, Struct.Schema.Batch) {}
-	FMemberReader(const FStructSchema& Schema, FByteReader Values, FSchemaBatchId InBatch);
+	PLAINPROPS_API FMemberReader(const FStructSchema& Schema, FByteReader Values, FSchemaBatchId InBatch);
 
 	bool								HasMore() const			{ return MemberIdx < NumMembers; }
 	bool								HasSuper() const		{ return UsesSuper && MemberIdx == 0; }		
@@ -642,6 +642,7 @@ public:
 	PLAINPROPS_API FStructView			GrabStruct();				// @pre PeekKind() == EMemberKind::Struct
 	//FAnyMemberView					GrabAny();					// @pre HasMore()
 
+	PLAINPROPS_API void					SkipMember();				// @pre HasMore()
 
 	// Experimental 
 	// @pre Has N more contiguous members of the expected leaf type
diff --git a/Engine/Plugins/Experimental/PlainProps/Source/Public/PlainPropsSaveMember.h b/Engine/Plugins/Experimental/PlainProps/Source/Public/PlainPropsSaveMember.h
--- a/Engine/Plugins/Experimental/PlainProps/Source/Public/PlainPropsSaveMember.h
+++ b/Engine/Plugins/Experimental/PlainProps/Source/Public/PlainPropsSaveMember.h
@@ -467,4 +467,14 @@ public:
 	using FRangeSaverBase::Finalize;
 };
 
+//////////////////////////////////////////////////////////////////////////
+
+PLAINPROPS_API void							SaveMember(FMemberBuilder& Out, const void* Struct, FMemberId Name, const FSaveContext& Ctx, FLeafMemberBinding Member);
+PLAINPROPS_API void							SaveMember(FMemberBuilder& Out, const void* Struct, FMemberId Name, const FSaveContext& Ctx, FRangeMemberBinding Member);
+PLAINPROPS_API void							SaveMember(FMemberBuilder& Out, const void* Struct, FMemberId Name, const FSaveContext& Ctx, FStructMemberBinding Member);
+
+//////////////////////////////////////////////////////////////////////////
+
+[[nodiscard]] PLAINPROPS_API FMemberSchema	CreateRangeSchema(FScratchAllocator& Scratch, FRangeMemberBinding Member);
+
 } // namespace PlainProps
diff --git a/Engine/Plugins/Experimental/PlainProps/Source/Public/PlainPropsVisitMember.h b/Engine/Plugins/Experimental/PlainProps/Source/Public/PlainPropsVisitMember.h
new file mode 100644
--- /dev/null
+++ b/Engine/Plugins/Experimental/PlainProps/Source/Public/PlainPropsVisitMember.h
@@ -0,0 +1,46 @@
+// Copyright Epic Games, Inc. All Rights Reserved.
+
+#pragma once
+
+#include "PlainPropsBind.h"
+
+namespace PlainProps 
+{
+
+// Iterates over member bindings
+class FMemberVisitor
+{
+public:
+	PLAINPROPS_API explicit FMemberVisitor(const FSchemaBinding& InSchema);
+
+	bool									HasMore() const			{ return MemberIdx < NumMembers; }
+	uint16									GetIndex() const		{ return MemberIdx; }
+	
+	PLAINPROPS_API EMemberKind				PeekKind() const;		// @pre HasMore()
+	PLAINPROPS_API FMemberBindType			PeekType() const;		// @pre HasMore()
+	PLAINPROPS_API uint32					PeekOffset() const;		// @pre HasMore()
+
+	PLAINPROPS_API FLeafMemberBinding		GrabLeaf();				// @pre PeekKind() == EMemberKind::Leaf
+	PLAINPROPS_API FRangeMemberBinding		GrabRange();			// @pre PeekKind() == EMemberKind::Range
+	PLAINPROPS_API FStructMemberBinding		GrabStruct();			// @pre PeekKind() == EMemberKind::Struct
+	PLAINPROPS_API FBindId					GrabSuper();			// @pre First grab and has declared super
+	PLAINPROPS_API void						SkipMember();
+
+protected: // for unit tests
+	const FSchemaBinding&		Schema;
+	const uint16				NumMembers;
+	uint16						MemberIdx = 0;
+	uint16						InnerRangeIdx = 0;		// Types of [nested] ranges
+	uint16						InnerIdIdx = 0;			// Types of static structs and enums
+
+	using FMemberBindTypeRange = TConstArrayView<FMemberBindType>;
+
+	PLAINPROPS_API uint64					GrabMemberOffset();
+	PLAINPROPS_API FMemberBindTypeRange		GrabInnerTypes();
+	PLAINPROPS_API FInnerId					GrabInnerId();
+	PLAINPROPS_API FBindId					GrabStructSchema(FStructType Type);
+	PLAINPROPS_API FOptionalInnerId			GrabRangeSchema(FMemberType InnermostType);
+	PLAINPROPS_API FEnumId					GrabEnumSchema()		{ return GrabInnerId().AsEnum(); }
+};
+
+} // namespace PlainProps
\ No newline at end of file
diff --git a/Engine/Plugins/Experimental/PlainPropsUObject/Source/Private/PlainPropsRestoreOverridable.cpp b/Engine/Plugins/Experimental/PlainPropsUObject/Source/Private/PlainPropsRestoreOverridable.cpp
new file mode 100644
--- /dev/null
+++ b/Engine/Plugins/Experimental/PlainPropsUObject/Source/Private/PlainPropsRestoreOverridable.cpp
@@ -0,0 +1,181 @@
+// Copyright Epic Games, Inc. All Rights Reserved.
+
+#include "PlainPropsRestoreOverridableInternal.h"
+#include "PlainPropsBind.h"
+#include "PlainPropsUObjectRuntime.h"
+#include "Containers/Array.h"
+#include "Serialization/ArchiveSerializedPropertyChain.h"
+#include "UObject/Class.h"
+#include "UObject/OverriddenPropertySet.h"
+#include "UObject/UnrealType.h"
+
+namespace PlainProps::UE
+{
+
+// Todo: This should be moved to an appropriate context that we can pass around.
+thread_local FOverriddenPropertySet* TlsRestoreOverrides = nullptr;
+
+// Determines whether an override operation is Modified or Replaced. A struct member is considered 
+// replaced if all of its members have been serialized. 
+// Note: in FOverriddenPropertySet::NotifyPropertyChange we consider an edited property as
+// 'Replaced' when StructProperty->Struct->UseNativeSerialization() is true.
+static EOverriddenPropertyOperation DetermineOverrideOperation(FStructView SerializedStruct, 
+															   const FRestoreContext& Ctx)
+{
+	FMemberReader It(SerializedStruct);
+
+	uint16 NumSerializedMembers = 0;
+
+	while (It.HasMore())
+	{
+		if (It.PeekKind() == EMemberKind::Struct)
+		{
+			EOverriddenPropertyOperation SubOp = DetermineOverrideOperation(It.GrabStruct(), Ctx);
+			if (SubOp == EOverriddenPropertyOperation::Modified)
+			{
+				// When the sub-operation is Modified, then this member is not fully
+				// replaced and more sub-operations will be nested under it.
+				return EOverriddenPropertyOperation::Modified;
+			}
+		}
+		else
+		{
+			It.SkipMember();
+		}
+
+		++NumSerializedMembers;
+	}
+
+	// This is a temporary workaround to find whether we have serialized all members or not.
+	// This could be optimized but we could also be using LoadIds instead.
+	// Note: assume we are using ESchemaFormat::StableNames.
+	check(SerializedStruct.Schema.Id.Idx < static_cast<uint32>(Ctx.RuntimeIds.Num()));
+	FBindId BindId = FBindId(reinterpret_cast<const FDeclId*>(Ctx.RuntimeIds.GetData())[SerializedStruct.Schema.Id.Idx]);
+
+	uint16 NumDeclMembers = 0;
+	const FStructDeclaration* Decl = nullptr;
+	if (Ctx.Schemas.FindStruct(BindId, Decl))
+	{
+		NumDeclMembers = Decl->NumMembers;
+	}
+	else if (const FStructDeclaration* CustomDecl = Ctx.Customs.FindDeclaration(BindId))
+	{
+		NumDeclMembers = CustomDecl->NumMembers;
+	}
+
+	return (NumSerializedMembers == NumDeclMembers) ? EOverriddenPropertyOperation::Replace : EOverriddenPropertyOperation::Modified;
+}
+
+static void RestoreStruct(const UStruct* Struct, FStructView StructView, FOverriddenPropertySet& Overrides, FArchiveSerializedPropertyChain& PropertyChain, const FRestoreContext& Ctx);
+
+static void RestoreProperty(FProperty* Property, 
+							FMemberType Type, 
+							FArchiveSerializedPropertyChain& PropertyChain,
+							FOverriddenPropertySet& Overrides,
+							FMemberReader& Reader,
+							const FRestoreContext& Ctx)
+{
+	FStructView SubStructView;
+
+	// Properties with the ExperimentalOverridableLogic flag have already been restored
+	// during loading via a custom binding.
+	bool bSkipOverride = Property->HasAnyPropertyFlags(CPF_ExperimentalOverridableLogic);
+
+	EOverriddenPropertyOperation Op = EOverriddenPropertyOperation::Replace;
+
+	switch (Type.GetKind())
+	{
+		case EMemberKind::Leaf:
+			Reader.GrabLeaf();
+			break;
+
+		case EMemberKind::Struct:
+		{
+			SubStructView = Reader.GrabStruct();
+
+			if (!bSkipOverride)
+			{
+				Op = DetermineOverrideOperation(SubStructView, Ctx);
+			}
+		}
+		break;
+
+		case EMemberKind::Range:
+			Reader.GrabRange();
+			break;
+	}
+
+	if (!bSkipOverride)
+	{
+		check(Op != EOverriddenPropertyOperation::None);
+
+		Overrides.RestoreOverriddenPropertyOperation(Op, &PropertyChain, Property);
+
+		if (Type.GetKind() == EMemberKind::Struct)
+		{
+			if (const FStructProperty* StructProperty = CastField<FStructProperty>(Property))
+			{
+				PropertyChain.PushProperty(Property, Property->IsEditorOnlyProperty());
+				RestoreStruct(StructProperty->Struct, SubStructView, Overrides, PropertyChain, Ctx);
+				PropertyChain.PopProperty(Property, Property->IsEditorOnlyProperty());
+			}
+		}		
+	}
+}
+
+static void RestoreStruct(const UStruct* Struct, 
+							FStructView StructView,
+							FOverriddenPropertySet& Overrides,
+							FArchiveSerializedPropertyChain& PropertyChain,
+							const FRestoreContext& Ctx)
+{	
+	FMemberReader It(StructView);
+
+	while (It.HasMore())
+	{
+		if (FOptionalMemberId Name = It.PeekName())
+		{
+			FMemberType Type = It.PeekType();
+
+			FName MemberFName = GUE.Names.ResolveName(Name.Get().Id).ToName();
+			checkSlow(MemberFName != NAME_None);
+
+			if (FProperty* Property = Struct->FindPropertyByName(MemberFName))
+			{
+				RestoreProperty(Property, Type, PropertyChain,
+								Overrides, It, Ctx);
+
+				continue;
+			}
+		}
+
+		It.SkipMember();
+	}
+}
+
+void RestoreStructOverrides(const UStruct* Struct, FStructView StructView, 
+							FOverriddenPropertySet& Overrides, 
+							const FRestoreContext& Ctx)
+{
+	FArchiveSerializedPropertyChain PropertyChain;
+
+	RestoreStruct(Struct, StructView, Overrides, PropertyChain, Ctx);
+}
+
+FScopeRestoreOverrides::FScopeRestoreOverrides(FOverriddenPropertySet* OverriddenProperties)
+	: PrevOverrides(TlsRestoreOverrides)
+{
+	TlsRestoreOverrides = OverriddenProperties;
+}
+
+FScopeRestoreOverrides::~FScopeRestoreOverrides()
+{
+	TlsRestoreOverrides = PrevOverrides;
+}
+
+FOverriddenPropertySet* GetRestoreOverrides()
+{
+	return TlsRestoreOverrides;
+}
+
+} // namespace PlainProps::UE
diff --git a/Engine/Plugins/Experimental/PlainPropsUObject/Source/Private/PlainPropsRestoreOverridableInternal.h b/Engine/Plugins/Experimental/PlainPropsUObject/Source/Private/PlainPropsRestoreOverridableInternal.h
new file mode 100644
--- /dev/null
+++ b/Engine/Plugins/Experimental/PlainPropsUObject/Source/Private/PlainPropsRestoreOverridableInternal.h
@@ -0,0 +1,41 @@
+// Copyright Epic Games, Inc. All Rights Reserved.
+
+#pragma once
+
+#include "PlainPropsRead.h"
+#include "PlainPropsTypes.h"
+
+struct FOverriddenPropertySet;
+class UStruct;
+
+namespace PlainProps
+{
+	class FSchemaBindings;
+	class FCustomBindings;
+}
+
+namespace PlainProps::UE
+{
+
+struct FRestoreContext
+{
+	const FSchemaBindings& Schemas;
+	FCustomBindings& Customs;
+	TConstArrayView<FStructId> RuntimeIds;
+};
+
+void RestoreStructOverrides(const UStruct* Struct, FStructView StructView, FOverriddenPropertySet& Overrides, const FRestoreContext& Ctx);
+
+class FScopeRestoreOverrides
+{
+public:
+	FScopeRestoreOverrides(FOverriddenPropertySet* OverriddenProperties);
+	~FScopeRestoreOverrides();
+
+private:
+	FOverriddenPropertySet* PrevOverrides;
+};
+
+FOverriddenPropertySet* GetRestoreOverrides();
+
+} // namespace PlainProps::UE
diff --git a/Engine/Plugins/Experimental/PlainPropsUObject/Source/Private/PlainPropsRoundtripTest.cpp b/Engine/Plugins/Experimental/PlainPropsUObject/Source/Private/PlainPropsRoundtripTest.cpp
--- a/Engine/Plugins/Experimental/PlainPropsUObject/Source/Private/PlainPropsRoundtripTest.cpp
+++ b/Engine/Plugins/Experimental/PlainPropsUObject/Source/Private/PlainPropsRoundtripTest.cpp
@@ -10,6 +10,8 @@
 #include "PlainPropsPrint.h"
 #include "PlainPropsSpecify.h"
 #include "PlainPropsWrite.h"
+#include "PlainPropsSaveOverridableInternal.h"
+#include "PlainPropsRestoreOverridableInternal.h"
 #include "Algo/Find.h"
 #include "Async/ParallelFor.h"
 #include "HAL/FileManager.h"
@@ -34,6 +36,7 @@
 #include "UObject/TextProperty.h"
 #include "UObject/UObjectIterator.h"
 #include "UObject/UObjectGlobals.h"
+#include "UObject/OverridableManager.h"
 #include "ProfilingDebugging/CpuProfilerTrace.h"
 
 #define PP_WITH_ENGINE_HOOKS 0
@@ -164,7 +167,10 @@ struct TCustomPropertyBinding final : ICustomBinding
 	virtual void SaveCustom(FMemberBuilder& Dst, const void* Src, FBaseline Base, const FSaveContext& Ctx) override
 	{
 		const void* Default = Base.Get();
-		if (!Default || DiffCustom(Src, Default, Ctx))
+		
+		// Do not perform a diff if we are saving via the override system because
+		// we want to save despite the value being the same as the default.
+		if (!Default || (GetSaveOverrides() != nullptr) || DiffCustom(Src, Default, Ctx))
 		{
 			Batch.Save(Dst, *static_cast<const Type*>(Src), Ctx);
 		}
@@ -304,6 +310,15 @@ public:
 		ObjectMap.Add(SavedObjects.Num());
 		SavedObjects.Emplace(Id, Built, Object);
 	}
+
+	void SaveObjectWithOverrides(FBindId Id, const UObject* Object, 
+								 const FOverriddenPropertySet& Overriddes, const UObject* Arch)
+	{
+		FBuiltStruct* Built = SaveStructOverrides(Overriddes, FBaseline(Arch), Id, DeltaCtx);
+		
+		ObjectMap.Add(SavedObjects.Num());
+		SavedObjects.Emplace(Id, Built, Object);
+	}
 	
 	void Skip(FBindId Id, const UObject* Object)
 	{
@@ -573,6 +588,16 @@ public:
 		++NumLoaded;
 	}
 
+	FORCENOINLINE void RestoreObjectOverrides(UObject* Object, int32 ObjectIndex, FOverriddenPropertySet& Overrides)
+	{
+		check(HasObject(ObjectIndex));
+		FStructView StructView = Objects[ObjectMap[ObjectIndex]];
+
+		FRestoreContext Ctx{ GUE.Schemas, GUE.Customs, MakeConstArrayView(LoadStructIds) };
+
+		RestoreStructOverrides(Object->GetClass(), StructView, Overrides, Ctx);
+	}
+
 	static void PrintDebugInfo()
 	{
 		UE_LOGFMT(LogPlainPropsUObject, Display,
@@ -1433,7 +1458,14 @@ struct FPlainPropsLinkerSandbox
 			FBindId Id = GUE.Names.IndexBindId(PlainProps::UE::IndexStruct(Struct));
 			if (GUE.Schemas.FindStruct(Id))
 			{
-				BatchSaver.Save(Id, Export.Object, Export.Object->GetArchetype());
+				if (const FOverriddenPropertySet* Overriddes = FOverridableManager::Get().GetOverriddenProperties(Export.Object))
+				{
+					BatchSaver.SaveObjectWithOverrides(Id, Export.Object, *Overriddes, Export.Object->GetArchetype());
+				}
+				else
+				{
+					BatchSaver.Save(Id, Export.Object, Export.Object->GetArchetype());
+				}
 			}
 			else
 			{
@@ -1614,7 +1646,19 @@ struct FPlainPropsLinkerSandbox
 		}
 		{
 			TRACE_CPUPROFILER_EVENT_SCOPE(LoadPlainProps);
-			BatchLoader.LoadObject(Object, ExportIndex);
+
+			if (FOverriddenPropertySet* Overrides = FOverridableManager::Get().GetOverriddenProperties(Object))
+			{
+				FScopeRestoreOverrides RestoreScope(Overrides);
+
+				BatchLoader.LoadObject(Object, ExportIndex);
+
+				BatchLoader.RestoreObjectOverrides(Object, ExportIndex, *Overrides);
+			}
+			else
+			{
+				BatchLoader.LoadObject(Object, ExportIndex);
+			}
 		}
 
 #if WITH_EDITOR
@@ -2323,6 +2367,27 @@ int32 RoundtripViaLinkerBatch(
 				{
 					++NumIgnored;
 				}
+
+				if (Ctx.OverrideDiffs.Num())
+				{
+					UE_LOGFMT(LogPlainPropsUObject, Display, "");
+					UE_LOGFMT(LogPlainPropsUObject, Display, "Detected override diff in package {PackageName}",
+						PackageNames[I]);
+
+					FOverridableManager& OverrideManager = FOverridableManager::Get();
+
+					for (FDiffOverride& Diff : Ctx.OverrideDiffs)
+					{
+						const FOverriddenPropertySet* OverridesA = OverrideManager.GetOverriddenProperties(Diff.A);
+						const FOverriddenPropertySet* OverridesB = OverrideManager.GetOverriddenProperties(Diff.B);
+
+						UE_LOGF(LogPlainPropsUObject, Display, "A: %ls", *Diff.A->GetPathName());
+						UE_LOGF(LogPlainPropsUObject, Display, "A overrides:");
+						UE_LOGF(LogPlainPropsUObject, Display, "\n%ls", OverridesA ? *OverridesA->ToDebugString() : TEXT("NULL"));
+						UE_LOGF(LogPlainPropsUObject, Display, "B overrides");
+						UE_LOGF(LogPlainPropsUObject, Display, "\n%ls", OverridesB ? *OverridesB->ToDebugString() : TEXT("NULL"));
+					}
+				}
 			}, Flags);
 			UE_LOGFMT(LogPlainPropsUObject, Display, "");
 			UE_LOGFMT(LogPlainPropsUObject, Display, "Ignored {NumIgnored} package diff(s) in {Num} package(s)", NumIgnored, NumPackages);
diff --git a/Engine/Plugins/Experimental/PlainPropsUObject/Source/Private/PlainPropsSaveOverridable.cpp b/Engine/Plugins/Experimental/PlainPropsUObject/Source/Private/PlainPropsSaveOverridable.cpp
new file mode 100644
--- /dev/null
+++ b/Engine/Plugins/Experimental/PlainPropsUObject/Source/Private/PlainPropsSaveOverridable.cpp
@@ -0,0 +1,207 @@
+// Copyright Epic Games, Inc. All Rights Reserved.
+
+#include "PlainPropsSaveOverridableInternal.h"
+#include "PlainPropsInternalPrivateMemberPtr.h"
+#include "PlainPropsSave.h"
+#include "PlainPropsSaveMember.h"
+#include "PlainPropsVisitMember.h"
+#include "PlainPropsUObjectRuntime.h"
+#include "UObject/OverriddenPropertySet.h"
+
+PP_DEFINE_PRIVATE_MEMBER_NAMESPACE_BEGIN(FOverriddenPropertyNodeID_Private)
+PP_DEFINE_PRIVATE_MEMBER_PTR(FOverriddenPropertyNodeID, Path, FOverriddenPropertyPath);
+PP_DEFINE_PRIVATE_MEMBER_NAMESPACE_END()
+
+PP_DEFINE_PRIVATE_MEMBER_NAMESPACE_BEGIN(FOverriddenPropertyPath_Private)
+PP_DEFINE_PRIVATE_MEMBER_PTR(FOverriddenPropertyPath, Path, TArray<FName>);
+PP_DEFINE_PRIVATE_MEMBER_NAMESPACE_END()
+
+namespace PlainProps::UE
+{
+
+// Todo: This should be moved to an appropriate context that we can pass around.
+thread_local const FOverriddenPropertySet* TlsSaveOverrides = nullptr;
+
+inline const FOverriddenPropertyNode* FindSubOverride(const FOverriddenPropertyNode& Override, FName MemberName)
+{
+	for (const FOverriddenPropertyNode& SubNode : Override.GetSubPropertyNodes())
+	{
+		const FOverriddenPropertyPath* Path = nullptr;
+
+		{
+			using namespace FOverriddenPropertyNodeID_Private;
+			Path = &(SubNode.GetNodeID().*_Path);
+		}
+		
+		{
+			using namespace FOverriddenPropertyPath_Private;
+			FName IdName = Path->IsEmpty() ? FName() : ((*Path).*_Path)[0];
+
+			if (IdName == MemberName)
+			{
+				return &SubNode;
+			}
+		}
+	}
+	return nullptr;
+}
+
+inline const uint8* At(const void* Ptr, SIZE_T Offset)
+{
+	return static_cast<const uint8*>(Ptr) + Offset;
+}
+
+inline bool IsOverriden(const FOverriddenPropertyNode& Override)
+{
+	EOverriddenPropertyOperation Op = Override.GetOperation();
+	return (Op != EOverriddenPropertyOperation::None) && (Op != EOverriddenPropertyOperation::SubObjectsShadowing);
+}
+
+inline void SaveNamedMembers(FMemberBuilder& Out, const void* Struct, FMemberVisitor It, TConstArrayView<FMemberId> Names, const FSaveContext& Ctx)
+{
+	for (FMemberId Name : Names)
+	{
+		switch (It.PeekKind())
+		{
+		case EMemberKind::Leaf:		SaveMember(Out, Struct, Name, Ctx, It.GrabLeaf());		break;
+		case EMemberKind::Range:	SaveMember(Out, Struct, Name, Ctx, It.GrabRange());		break;
+		case EMemberKind::Struct:	SaveMember(Out, Struct, Name, Ctx, It.GrabStruct());	break;
+		}
+	}
+	checkSlow(!It.HasMore());
+}
+
+static void SaveRangeMember(FMemberBuilder& Out, const void* Struct, FMemberId Name, const FSaveContext& Ctx, FRangeMemberBinding Member)
+{
+	Out.AddRange(Name, { CreateRangeSchema(Ctx.Scratch, Member), SaveRange(At(Struct, Member.Offset), Member, Ctx) });
+}
+
+static const FStructDeclaration& SaveStructOverrides(FMemberBuilder& Out, const void* Struct, const FOverriddenPropertyNode& Override, FBaseline Base, FBindId Id, const FSaveContext& Ctx);
+
+static void SaveStructMember(FMemberBuilder& Out, const void* Struct, const FOverriddenPropertyNode& Override, 
+							 const void* Default, FMemberId Name, const FSaveContext& Ctx, FStructMemberBinding Member)
+{
+	FMemberBuilder TempBuilder;
+	const void* MemberDefault = Default ? At(Default, Member.Offset) : nullptr;
+	const FStructDeclaration& Declaration = SaveStructOverrides(TempBuilder, At(Struct, Member.Offset), Override, MemberDefault, Member.Id, Ctx);
+
+	FBuiltStruct* Result = TempBuilder.IsEmpty() ? nullptr : TempBuilder.BuildAndReset(Ctx.Scratch, Declaration, Ctx.Schemas.GetDebug());
+
+	if (Result)
+	{
+		Out.AddStruct(Name, Member.Id, Result);
+	}
+}
+
+static void SaveSchemaBoundStructOverrides(FMemberBuilder& Out, const void* Struct, const FOverriddenPropertyNode& Override,
+											FBaseline Base, const FSchemaBinding& Schema, const FStructDeclaration& Declaration, 
+											const FSaveContext& Ctx)
+{
+	FMemberVisitor It(Schema);
+	if (Declaration.Super)
+	{
+		FBindId SuperId = It.GrabSuper();
+		checkSlow(SuperId == ToOptionalStruct(Declaration.Super));
+		const FStructDeclaration* SuperDecl;
+		const FSchemaBinding& SuperSchema = Ctx.Schemas.GetStruct(SuperId, /* out */ SuperDecl);
+		SaveSchemaBoundStructOverrides(Out, Struct, Override, Base.Super(), SuperSchema, *SuperDecl, Ctx);
+		Out.BuildSuperStruct(Ctx.Scratch, *SuperDecl, Ctx.Schemas.GetDebug());
+	}
+
+	if (IsOverriden(Override))
+	{
+		if (Declaration.Occupancy == EMemberPresence::AllowSparse)
+		{
+			for (FMemberId Name : Declaration.GetMemberOrder())
+			{
+				const FOverriddenPropertyNode* SubOverride = nullptr;
+				if (Override.GetOperation() == EOverriddenPropertyOperation::Replace)
+				{
+					// Keep using the same override since a replace operation implies no sub-nodes.
+					SubOverride = &Override;
+				}
+				else
+				{
+					FName MemberFName = GUE.Names.ResolveName(Name.Id).ToName();
+
+					SubOverride = FindSubOverride(Override, MemberFName);
+				}
+
+				if (SubOverride)
+				{
+					// The default is only used for custom bindings and specifically overridable container bindings.
+					// When the property operation is 'Replace' we pass no default so that the bindings for overridable
+					// containers do not consider individual items but instead a single 'Assign' operation.
+					const void* Default = (Override.GetOperation() == EOverriddenPropertyOperation::Replace)
+											? nullptr : Base.Get();
+
+					switch (It.PeekKind())
+					{
+					case EMemberKind::Leaf:		SaveMember(Out, Struct, Name, Ctx, It.GrabLeaf());	break;
+					case EMemberKind::Range:	SaveRangeMember(Out, Struct, Name, Ctx, It.GrabRange());	break;
+					case EMemberKind::Struct:	SaveStructMember(Out, Struct, *SubOverride, Default, Name, Ctx, It.GrabStruct()); break;
+					}
+				}
+				else
+				{
+					It.SkipMember();
+				}
+				
+			}
+			checkSlow(!It.HasMore());
+		}
+		else
+		{
+			SaveNamedMembers(Out, Struct, It, Declaration.GetMemberOrder(), Ctx);
+		}
+	}
+}
+
+static const FStructDeclaration& SaveStructOverrides(FMemberBuilder& Out, const void* Struct, const FOverriddenPropertyNode& Override,
+													 FBaseline Base, FBindId Id, const FSaveContext& Ctx)
+{
+	const FStructDeclaration* Declaration = nullptr;
+	if (ICustomBinding* Custom = Ctx.Customs.FindStructToSave(Id, /* out */ Declaration))
+	{
+		if (IsOverriden(Override))
+		{
+			Custom->SaveCustom(Out, Struct, Base, Ctx);
+		}
+	}
+	else
+	{
+		const FSchemaBinding& Schema = Ctx.Schemas.GetStruct(Id, /* out */ Declaration);
+		SaveSchemaBoundStructOverrides(Out, Struct, Override, Base, Schema, *Declaration, Ctx);
+	}
+
+	return *Declaration;
+}
+
+static FBuiltStruct* SaveStructOverrides(const void* Struct, const FOverriddenPropertyNode& Override, FBaseline Base, FBindId Id, const FSaveContext& Ctx)
+{
+	FMemberBuilder Out;
+	const FStructDeclaration& Declaration = SaveStructOverrides(Out, Struct, Override, Base, Id, Ctx);
+	return Out.BuildAndReset(Ctx.Scratch, Declaration, Ctx.Schemas.GetDebug());
+}
+
+FBuiltStruct* SaveStructOverrides(const FOverriddenPropertySet& Overrides, 
+								  FBaseline Base, FBindId BindId, const FSaveContext& Ctx)
+{
+	const FOverriddenPropertySet* PrevOverrides = TlsSaveOverrides;
+	TlsSaveOverrides = &Overrides;
+	
+	const FOverriddenPropertyNode& RootOverride = Overrides.GetRootOverriddenPropertyNode();
+
+	FBuiltStruct* Result = SaveStructOverrides(Overrides.GetOwner(), RootOverride, Base, BindId, Ctx);
+
+	TlsSaveOverrides = PrevOverrides;
+
+	return Result;
+}
+
+const FOverriddenPropertySet* GetSaveOverrides()
+{
+	return TlsSaveOverrides;
+}
+
+} // namespace PlainProps::UE
diff --git a/Engine/Plugins/Experimental/PlainPropsUObject/Source/Private/PlainPropsSaveOverridableInternal.h b/Engine/Plugins/Experimental/PlainPropsUObject/Source/Private/PlainPropsSaveOverridableInternal.h
new file mode 100644
--- /dev/null
+++ b/Engine/Plugins/Experimental/PlainPropsUObject/Source/Private/PlainPropsSaveOverridableInternal.h
@@ -0,0 +1,21 @@
+// Copyright Epic Games, Inc. All Rights Reserved.
+
+#pragma once
+
+#include "PlainPropsTypes.h"
+#include "PlainPropsSave.h"
+
+struct FOverriddenPropertySet;
+
+namespace PlainProps
+{
+
+namespace UE
+{
+
+FBuiltStruct* SaveStructOverrides(const FOverriddenPropertySet& Overrides, FBaseline Base, FBindId BindId, const FSaveContext& Context);
+
+const FOverriddenPropertySet* GetSaveOverrides();
+
+} // namespace UE
+} // namespace PlainProps
diff --git a/Engine/Plugins/Experimental/PlainPropsUObject/Source/Private/PlainPropsUObjectBindings.cpp b/Engine/Plugins/Experimental/PlainPropsUObject/Source/Private/PlainPropsUObjectBindings.cpp
--- a/Engine/Plugins/Experimental/PlainPropsUObject/Source/Private/PlainPropsUObjectBindings.cpp
+++ b/Engine/Plugins/Experimental/PlainPropsUObject/Source/Private/PlainPropsUObjectBindings.cpp
@@ -6,6 +6,8 @@
 #include "PlainPropsUObjectModule.h"
 #include "PlainPropsUeCoreBindings.h"
 #include "PlainPropsVisualize.h"
+#include "PlainPropsRestoreOverridableInternal.h"
+#include "PlainPropsSaveOverridableInternal.h"
 #include "Algo/Compare.h"
 #include "Algo/Find.h"
 #include "Containers/PagedArray.h"
@@ -15,6 +17,7 @@
 #include "Misc/DefinePrivateMemberPtr.h"
 #include "StructUtils/UserDefinedStruct.h"
 #include "Templates/SharedPointer.h"
+#include "Serialization/ArchiveSerializedPropertyChain.h"
 #include "Templates/UniquePtr.h"
 #include "UObject/AnsiStrProperty.h"
 #include "UObject/Class.h"
@@ -31,6 +34,7 @@
 #include "UObject/Utf8StrProperty.h"
 #include "UObject/VerseStringProperty.h"
 #include "UObject/VerseValueProperty.h"
+#include "UObject/OverriddenPropertySet.h"
 #include "VerseVM/VVMNativeString.h"
 
 
@@ -306,6 +310,7 @@ FCommonMemberIds::FCommonMemberIds(TIdIndexer<FSensitiveName>& Names)
 , Assign(					Names.NameMember("Assign"))
 , Remove(					Names.NameMember("Remove"))
 , Insert(					Names.NameMember("Insert"))
+, Modify(					Names.NameMember("Modify"))
 , Id(						Names.NameMember("Id"))
 , Object(					Names.NameMember("Object"))
 , Function(					Names.NameMember("Function"))
@@ -730,6 +735,107 @@ static FBindId BindSparseDelegate(FBindId Owner, FMulticastSparseDelegatePropert
 
 ////////////////////////////////////////////////////////////////////////////////////////////////
 
+// Todo: Temporary utilities that allow us to access overridable serialization state with just a pointer
+// to property data. These should be replaced by keeping track of the current override node.
+
+static void FindPropertyWithData(UObject* Owner, const void* PropertyData,
+								 const TFunctionRef<void(const FPropertyVisitorContext& /*Context*/)> PropertyMatchFunc)
+{
+	TSet<const FProperty*> VisitedProperties;
+
+	Owner->GetClass()->Visit(Owner, [&VisitedProperties, PropertyData, &PropertyMatchFunc](const FPropertyVisitorContext& Context)->EPropertyVisitorControlFlow
+	{
+		const FPropertyVisitorPath& PropertyPath = Context.Path;
+		const FPropertyVisitorData& Data = Context.Data;
+
+		const FProperty* Property = PropertyPath.Top().Property;
+
+		if (!Property || VisitedProperties.Contains(Property))
+		{
+			return EPropertyVisitorControlFlow::StepOver;
+		}
+
+		if (Data.PropertyData == PropertyData)
+		{
+			// We don't stop the visit here because we want to find the innermost matching property.
+			PropertyMatchFunc(Context);
+		}
+
+		VisitedProperties.Emplace(Property);
+
+		// Don't step into references to other objects to avoid cycles.
+		if (const FObjectPropertyBase* ObjectProperty = CastField<FObjectPropertyBase>(Property))
+		{
+			return EPropertyVisitorControlFlow::StepOver;
+		}
+		else
+		{
+			return EPropertyVisitorControlFlow::StepInto;
+		}
+	});
+}
+
+static const FOverriddenPropertyNode* GetOverriddenNode(const FOverriddenPropertySet& Overrides, const void* PropertyData)
+{
+	const FOverriddenPropertyNode* CurrentNode = nullptr;
+
+	FindPropertyWithData(Overrides.GetOwner(), PropertyData, [&CurrentNode, &Overrides](const FPropertyVisitorContext& Context)
+	{
+		const FPropertyVisitorPath& PropertyPath = Context.Path;
+
+		FArchiveSerializedPropertyChain PropertyChain = PropertyPath.ToSerializedPropertyChain();
+
+		if (const FOverriddenPropertyNode* Node = Overrides.GetOverriddenPropertyNode(&PropertyChain))
+		{
+			CurrentNode = Node;
+		}
+	});
+
+	return CurrentNode;
+}
+
+static EOverriddenPropertyOperation GetOverriddenOperation(const FOverriddenPropertySet& Overrides, const void* PropertyData)
+{
+	FPropertyVisitorPath FoundPath;
+
+	FindPropertyWithData(Overrides.GetOwner(), PropertyData, [&FoundPath](const FPropertyVisitorContext& Context)
+	{
+		FoundPath = Context.Path;
+	});
+
+	EOverriddenPropertyOperation Operation = EOverriddenPropertyOperation::None;
+
+	if (FoundPath.Num() > 0)
+	{
+		Operation = Overrides.GetOverriddenPropertyOperation(FoundPath.GetRootIterator());
+	}
+
+	return Operation;
+}
+
+static FOverriddenPropertyNode* RestoreOverriddenOperation(FOverriddenPropertySet& Overrides, EOverriddenPropertyOperation Operation, const void* PropertyData)
+{
+	FPropertyVisitorPath FoundPath;
+
+	FindPropertyWithData(Overrides.GetOwner(), PropertyData, [&FoundPath](const FPropertyVisitorContext& Context)
+	{
+		FoundPath = Context.Path;
+	});
+
+	FOverriddenPropertyNode* OverriddenNode = nullptr;
+
+	if (FoundPath.Num() > 0)
+	{
+		FArchiveSerializedPropertyChain PropertyChain = FoundPath.ToSerializedPropertyChain();
+		OverriddenNode = Overrides.RestoreOverriddenPropertyOperation(Operation, &PropertyChain, /*Property*/nullptr);
+
+	}
+
+	return OverriddenNode;
+}
+
+////////////////////////////////////////////////////////////////////////////////////////////////
+
 FVerseFunctionBinding::FVerseFunctionBinding(TPropertySpecifier<1>& Spec)
 : MemberIds{GUE.Members.Value}
 {
@@ -1178,6 +1284,9 @@ struct TLeafArrayBinding : ILeafRangeBinding
 	}
 };
 
+struct FOverridableArrayPropertyBinding;
+inline FOverridableArrayPropertyBinding* CreateAndLeakOverridableArrayBinding(FArrayProperty* ArrayProperty, FMemberBinding InnerBinding, FBindId BindId, FDeclId DeclId);
+
 // Reusable cache of FArrayProperty range bindings
 class FArrayPropertyBindings
 {
@@ -1197,6 +1306,7 @@ class FArrayPropertyBindings
 	const FRangeBinding											Integers[2][4];
 
 	TMap<FArrayPropertyInfo, IItemRangeBinding*>				Others;
+	TMap<FParameterBinding, FBindId>							OverridableBindings;
 
 	inline const ILeafRangeBinding& DownCast(const ILeafRangeBinding& In) { return In; }
 
@@ -1242,6 +1352,31 @@ public:
 		return FRangeBinding(*New, SizeType);
 	}
 
+	FBindId BindOverridableArray(FArrayProperty* Property, FMemberBinding InnerBinding)
+	{
+		FProperty* Inner = Property->Inner;
+		FArrayPropertyInfo Info(Property);
+
+		check(Info.IsTrivial());
+
+		if (const FBindId* BindId = OverridableBindings.Find(FParameterBinding(InnerBinding)))
+		{
+			return *BindId;
+		}
+
+		FBothType Param = InnerBinding.IndexParameterName(GUE.Names);
+		FType BindType = FType{ GUE.Scopes.CoreUObject, FTypenameId(GUE.Names.MakeParametricTypeId(GUE.Typenames.TrivialArray, MakeArrayView(&Param.BindType, 1))) };
+		FType DeclType = Param.IsLowered()
+			? FType{ GUE.Scopes.CoreUObject, FTypenameId(GUE.Names.MakeParametricTypeId(GUE.Typenames.TrivialArray, MakeArrayView(&Param.DeclType, 1))) }
+		: BindType;
+		FBindId BindId = GUE.Names.IndexBindId(BindType);
+		FDeclId DeclId = Param.IsLowered() ? GUE.Names.IndexDeclId(DeclType) : LowerCast(BindId);
+
+		CreateAndLeakOverridableArrayBinding(Property, InnerBinding, BindId, DeclId);
+
+		return OverridableBindings.Emplace(InnerBinding, BindId);
+	}
+
 	template<typename ScriptArray>
 	IItemRangeBinding* CreateAndCache(FArrayPropertyInfo Info)
 	{
@@ -1870,6 +2005,45 @@ void SaveSetDelta(const BindingType& Binding, FMemberBuilder& Dst, const ScriptT
 	}
 }
 
+template<typename BindingType, typename ScriptType>
+void SaveSetOverrides(const BindingType& Binding, FMemberBuilder& Dst, const ScriptType& Src, const ScriptType* Default, const FSaveContext& Ctx,
+					  const FOverriddenPropertySet* Overrides)
+{
+	if (Overrides && Default)
+	{
+		TBitArray<> RemoveIds(false, Default->GetMaxIndex());
+		TBitArray<> InsertIds(false, Src.GetMaxIndex());
+		TBitArray<> ModifyIds(false, Src.GetMaxIndex());
+		bool bIsFullyReplaced = false;
+
+		Binding.SaveOverrides(Dst, Src, *Default, *Overrides, RemoveIds, InsertIds, ModifyIds, bIsFullyReplaced);
+
+		if (bIsFullyReplaced)
+		{
+			Dst.AddRange(GUE.Members.Assign, SaveAll(Src, Binding.Layout, Binding.GetItemRange(), Ctx));
+		}
+		else
+		{
+			if (typename BindingType::SubSetIterator Removed{ Binding.Layout, *Default, RemoveIds })
+			{
+				Dst.AddRange(GUE.Members.Remove, SaveSome(Removed, Binding.GetKeyRange(), Ctx));
+			}
+			if (typename BindingType::SubSetIterator Inserted{ Binding.Layout, Src, InsertIds })
+			{
+				Dst.AddRange(GUE.Members.Insert, SaveSome(Inserted, Binding.GetItemRange(), Ctx));
+			}
+			if (typename BindingType::SubSetIterator Modified{ Binding.Layout, Src, ModifyIds })
+			{
+				Dst.AddRange(GUE.Members.Modify, SaveSome(Modified, Binding.GetItemRange(), Ctx));
+			}
+		}
+	}
+	else
+	{
+		SaveSetDelta(Binding, Dst, Src, Default, Ctx);
+	}
+}
+
 template<typename BindingType, typename ScriptType>
 void InsertSetItems(const BindingType& Binding, ScriptType& Dst, FRangeLoadView Items)
 {
@@ -1885,7 +2059,7 @@ void InsertSetItems(const BindingType& Binding, ScriptType& Dst, FRangeLoadView
 }
 
 template<typename BindingType, typename ScriptType>
-void LoadSetDelta(const BindingType& Binding, ScriptType& Dst, FStructLoadView Src)
+void LoadSetDelta(const BindingType& Binding, ScriptType& Dst, FStructLoadView Src, FOverriddenPropertySet* Overrides = nullptr)
 {
 	FMemberLoader Members(Src);
 	FOptionalMemberId Name = Members.PeekName();
@@ -1893,16 +2067,21 @@ void LoadSetDelta(const BindingType& Binding, ScriptType& Dst, FStructLoadView S
 	if (Name == GUE.Members.Insert)
 	{
 		InsertSetItems(Binding, Dst, Range);
-	} 
+
+		Binding.RestoreOverrides(Dst, Range, Overrides, EOverriddenPropertyOperation::Add);
+	}
 	else if (Name == GUE.Members.Assign)
 	{
 		Binding.DestroyAll(Dst);
 		Dst.Empty(Dst.Max(), Binding.Layout);
 		Binding.AssignEmpty(Dst, Range);
+
+		Binding.RestoreOverrides(Dst, Range, Overrides, EOverriddenPropertyOperation::Replace);
 	}
-	else
+	else if (Name == GUE.Members.Remove)
 	{
-		checkSlow(Name == GUE.Members.Remove);
+		Binding.RestoreOverrides(Dst, Range, Overrides, EOverriddenPropertyOperation::Remove);
+
 		Binding.Remove(Dst, Range);
 		if (Members.HasMore())
 		{
@@ -1910,6 +2089,12 @@ void LoadSetDelta(const BindingType& Binding, ScriptType& Dst, FStructLoadView S
 			InsertSetItems(Binding, Dst, Members.GrabRange());
 		}
 	}
+	else
+	{
+		checkSlow(Name == GUE.Members.Modify);
+
+		Binding.RestoreOverrides(Dst, Range, Overrides, EOverriddenPropertyOperation::Modified);
+	}
 
 	checkSlow(!Members.HasMore());
 }
@@ -2083,6 +2268,10 @@ struct TSetPropertyBinding : IItemRangeBinding, ICustomBinding
 			}
 		}
 	}
+
+	inline void RestoreOverrides(const FScriptSet& Set, FRangeLoadView Items, FOverriddenPropertySet* Overrides, EOverriddenPropertyOperation Operation) const
+	{
+	}
 	
 	virtual void ReadItems(FSaveRangeContext& Ctx) const override
 	{
@@ -2508,6 +2697,10 @@ struct TMapPropertyCustomBinding : TMapPropertyItemBinding<FScriptMap, KeyKind,
 		return false;
 	}
 
+	inline void RestoreOverrides(const FScriptMap& Map, FRangeLoadView Src, FOverriddenPropertySet* Overrides, EOverriddenPropertyOperation Operation) const
+	{
+	}
+
 	inline uint8* AddPair(FScriptMap& Dst, int32& OutIdx) const
 	{
 		OutIdx = Dst.AddUninitialized(Layout);
@@ -2684,6 +2877,157 @@ struct TMapPropertyCustomBinding : TMapPropertyItemBinding<FScriptMap, KeyKind,
 	}
 };
 
+template <EPropertyKind KeyKind, EPropertyKind ValueKind>
+struct TOverridableMapBinding : TMapPropertyCustomBinding<KeyKind, ValueKind>
+{
+	using Super = TMapPropertyCustomBinding<KeyKind, ValueKind>;
+	using Super::Layout;
+	using Super::InnerKey;
+	using Super::InnerValue;
+	using Super::KeyRange;
+	using Super::PairRange;
+
+	FMapProperty*	MapProperty;
+
+	TOverridableMapBinding(FMapProperty* Map, FMapMemberBindings Members)
+		: Super(Map, Members)
+		, MapProperty(Map)
+	{
+	}
+
+	FStructDeclarationPtr Declare(FDeclId Id) const
+	{
+		FMemberId Members[] = { GUE.Members.Assign, GUE.Members.Remove, GUE.Members.Insert, GUE.Members.Modify };
+		FMemberSpec KeysSpec = KeyRange.SpecMember();
+		FMemberSpec PairsSpec = PairRange.SpecMember();
+		return PlainProps::Declare({ Id, NoId, 0, EMemberPresence::AllowSparse, Members, {PairsSpec, KeysSpec, PairsSpec, PairsSpec} });
+	}
+
+	inline void SaveOverrides(FMemberBuilder& Dst, const FScriptMap& Map, const FScriptMap& Default,
+							  const FOverriddenPropertySet& Overrides,
+							  TBitArray<>& RemoveIds, TBitArray<>& InsertIds, TBitArray<>& ModifyIds, bool& bIsFullyReplaced) const
+	{
+		bIsFullyReplaced = false;
+
+		if (const FOverriddenPropertyNode* MapNode = GetOverriddenNode(Overrides, &Map))
+		{
+			if (MapNode->GetOperation() == EOverriddenPropertyOperation::Replace)
+			{
+				bIsFullyReplaced = true;
+			}
+			else
+			{
+				FScriptMapHelper MapHelper(MapProperty, &Map);
+				FScriptMapHelper DefaultMapHelper(MapProperty, &Default);
+
+				for (const FOverriddenPropertyNode& SubNode : MapNode->GetSubPropertyNodes())
+				{
+					switch (SubNode.GetOperation())
+					{
+					case EOverriddenPropertyOperation::Remove:
+					{
+						const int32 InternalIndex = SubNode.GetNodeID().ToMapInternalIndex(DefaultMapHelper);
+						if (InternalIndex != INDEX_NONE)
+						{
+							check(InternalIndex < Default.GetMaxIndex());
+							RemoveIds[InternalIndex] = 1;
+						}
+						break;
+					}
+					case EOverriddenPropertyOperation::Add:
+					{
+						const int32 InternalIndex = SubNode.GetNodeID().ToMapInternalIndex(MapHelper);
+						if (InternalIndex != INDEX_NONE)
+						{
+							check(InternalIndex < Map.GetMaxIndex());
+							InsertIds[InternalIndex] = 1;
+						}
+						break;
+					}
+					case EOverriddenPropertyOperation::Modified:
+					{
+						const int32 InternalIndex = SubNode.GetNodeID().ToMapInternalIndex(MapHelper);
+						if (InternalIndex != INDEX_NONE)
+						{
+							check(InternalIndex < Map.GetMaxIndex());
+							ModifyIds[InternalIndex] = 1;
+						}
+						break;
+					}
+					case EOverriddenPropertyOperation::Externalized:
+					{
+						// Todo: Handle external objects.
+						break;
+					}
+					default:
+						checkf(false, TEXT("Unsupported map operation"));
+						break;
+					}
+				}
+			}
+		}
+	}
+
+	inline void RestoreOverrides(const FScriptMap& Map, FRangeLoadView Src,
+								 FOverriddenPropertySet* Overrides, EOverriddenPropertyOperation Operation) const
+	{
+		if (Overrides)
+		{
+			if (Operation == EOverriddenPropertyOperation::Replace)
+			{
+				RestoreOverriddenOperation(*Overrides, EOverriddenPropertyOperation::Replace, &Map);
+			}
+			else if (Operation == EOverriddenPropertyOperation::Add ||
+					 Operation == EOverriddenPropertyOperation::Remove ||
+					 Operation == EOverriddenPropertyOperation::Modified)
+			{
+				TArray<uint8, TInlineAllocator<64>> Buffer;
+				Buffer.SetNumUninitialized(Layout.SetLayout.Size);
+				void* Tmp = Buffer.GetData();
+
+				InnerKey.InitItem(Tmp);
+				InnerValue.InitItem(static_cast<uint8*>(Tmp) + Layout.ValueOffset);
+
+				if (!Src.IsEmpty())
+				{
+					FOptionalSchemaId InnerLoadIds[2];
+					FStructRangeLoadView Structs = Src.AsStructs();
+					Structs.GetSchema().GetInnerLoadIds(/* out */ MakeArrayView(InnerLoadIds));
+					for (FStructLoadView Struct : Structs)
+					{
+						LoadStruct(Tmp, Struct);
+
+						// Need to fetch the ArrayNode every loop as the previous iteration might have reallocated the node.
+						if (FOverriddenPropertyNode* MapNode = RestoreOverriddenOperation(*Overrides, EOverriddenPropertyOperation::Modified, &Map))
+						{
+							FOverriddenPropertyNodeID KeyID = FOverriddenPropertyNodeID::FromMapKey(MapProperty->GetKeyProperty(), Tmp);
+							Overrides->RestoreSubPropertyOperation(Operation, *MapNode, KeyID);
+						}
+					}
+				}
+
+				InnerKey.DestroyItem(Tmp);
+				InnerValue.DestroyItem(static_cast<uint8*>(Tmp) + Layout.ValueOffset);
+			}
+			else
+			{
+				checkf(false, TEXT("Unsupported override operation for maps"));
+			}
+		}
+	}
+
+	virtual void SaveCustom(FMemberBuilder& Dst, const void* Src, FBaseline Base, const FSaveContext& Ctx) override
+	{
+		SaveSetOverrides(*this, Dst, *static_cast<const FScriptMap*>(Src), static_cast<const FScriptMap*>(Base.Get()), Ctx, GetSaveOverrides());
+	}
+
+	virtual void LoadCustom(void* Dst, FStructLoadView Src, ECustomLoadMethod Method) const override
+	{
+		check(Method == ECustomLoadMethod::Assign);
+		LoadSetDelta(*this, *static_cast<FScriptMap*>(Dst), Src, GetRestoreOverrides());
+	}
+};
+
 ////////////////////////////////////////////////////////////////////////////////////////////////
 
 class FMapBindings
@@ -2699,6 +3043,8 @@ class FMapBindings
 			return new TMapPropertyItemBinding<FFreezableScriptMap, KeyKind, ValueKind>(Property);
 		}
 
+		const bool bIsOverridable = Property->HasAnyPropertyFlags(CPF_ExperimentalOverridableLogic);
+
 		// Index custom delta binding struct name
 		FBothStructId BothPair;
 		TArray<FInnerStruct, TInlineAllocator<3>> Inners;
@@ -2726,6 +3072,11 @@ class FMapBindings
 				Inners.Emplace(GUE.Members.Assign, PairId);
 				Inners.Emplace(GUE.Members.Insert, PairId);
 			}
+
+			if (bIsOverridable)
+			{
+				Inners.Emplace(GUE.Members.Modify, PairId);
+			}
 		}
 		else
 		{
@@ -2741,9 +3092,20 @@ class FMapBindings
 		}
 
 		// Todo: Ownership / memory leak
-		auto Out = new TMapPropertyCustomBinding<KeyKind, ValueKind>(Property, Members);
-		GUE.Customs.BindStruct(BothPair.BindId, *Out, Out->Declare(BothPair.DeclId), Inners);
+		TMapPropertyCustomBinding<KeyKind, ValueKind>* Out = nullptr;
 
+		if (bIsOverridable)
+		{
+			auto Binding = new TOverridableMapBinding<KeyKind, ValueKind>(Property, Members);
+			GUE.Customs.BindStruct(BothPair.BindId, *Binding, Binding->Declare(BothPair.DeclId), Inners);
+			Out = Binding;
+		}
+		else
+		{
+			Out = new TMapPropertyCustomBinding<KeyKind, ValueKind>(Property, Members);
+			GUE.Customs.BindStruct(BothPair.BindId, *Out, Out->Declare(BothPair.DeclId), Inners);
+		}
+		
 		*OutCustomId = BothPair.BindId;
 		return Out;
 	}
@@ -2811,6 +3173,399 @@ static FMapBindings GMaps;
 
 ////////////////////////////////////////////////////////////////////////////////////////////////
 
+// Utilities to allow the overridable array binding to use the underlying array as a set.
+// This is so we can use the templated set utilities (Save, Load, Diff, etc).
+
+struct FOverridableScriptArrayLayout
+{
+	int32 NumBytesPerElement;
+	uint32 AlignmentOfElement;
+};
+
+class FOverridableScriptArray : public FScriptArray
+{
+public:
+
+	[[nodiscard]] int32 GetMaxIndex() const
+	{
+		return Num();
+	}
+	[[nodiscard]] void* GetData(int32 Index, const FOverridableScriptArrayLayout& Layout)
+	{
+		void* Elems = FScriptArray::GetData();
+		return static_cast<uint8*>(Elems) + (Index * Layout.NumBytesPerElement);
+	}
+	[[nodiscard]] const void* GetData(int32 Index, const FOverridableScriptArrayLayout& Layout) const
+	{
+		const void* Elems = FScriptArray::GetData();
+		return static_cast<const uint8*>(Elems) + (Index * Layout.NumBytesPerElement);
+	}
+	void Empty(int32 Slack, const FOverridableScriptArrayLayout& Layout)
+	{
+		FScriptArray::Empty(Slack, Layout.NumBytesPerElement, Layout.AlignmentOfElement);
+	}
+	int32 AddUninitialized(const FOverridableScriptArrayLayout& Layout)
+	{
+		return Add(1, Layout.NumBytesPerElement, Layout.AlignmentOfElement);
+	}
+	void RemoveAtUninitialized(const FOverridableScriptArrayLayout& Layout, int32 Index)
+	{
+		FScriptArray::Remove(Index, /*Count*/1, Layout.NumBytesPerElement, Layout.AlignmentOfElement);
+	}
+};
+
+struct FOverridableArrayPropertyBinding : IItemRangeBinding, ICustomBinding
+{
+	using SubSetIterator = TSubSetIterator<FOverridableScriptArray, FOverridableScriptArrayLayout>;
+
+	const FOverridableScriptArrayLayout						Layout;
+	const TInnerProperty<EPropertyKind::Struct>				Inner;
+	const TPropertyRangeSerializer<EPropertyKind::Struct>	Range;
+
+	const TPropertyRangeSerializer<EPropertyKind::Struct>& GetKeyRange() const { return Range; }
+	const TPropertyRangeSerializer<EPropertyKind::Struct>& GetItemRange() const { return Range; }
+
+	FOverridableArrayPropertyBinding(FArrayProperty* In, FMemberBinding Elem)
+		: IItemRangeBinding(GUE.Typenames.TrivialArray)
+		, Layout(In->Inner->GetElementSize(), In->Inner->GetMinAlignment())
+		, Inner(In->Inner)
+		, Range(Elem)
+	{
+		check(Inner.bHashable);
+		checkf(CastField<FObjectProperty>(Inner.Property) && Inner.Property->HasAnyPropertyFlags(CPF_PersistentInstance), 
+			TEXT("Only instanced sub-object arrays can be overridable"));
+	}
+
+	inline FStructDeclarationPtr DeclareCustom(FDeclId Id) const
+	{
+		FMemberId Members[] = { GUE.Members.Assign, GUE.Members.Remove, GUE.Members.Insert, GUE.Members.Modify };
+		FMemberSpec RangeSpec = Range.SpecMember();
+		return Declare({ Id, NoId, 0, EMemberPresence::AllowSparse, Members, {RangeSpec, RangeSpec, RangeSpec, RangeSpec} });
+	}
+
+	inline SIZE_T GetStride() const { return Layout.NumBytesPerElement; }
+
+	inline int32 FindIndex(const FOverridableScriptArray& Array, const void* Elem) const
+	{
+		int32 Num = Array.Num();
+		int32 Result = INDEX_NONE;
+
+		if (const UObject** Obj = (const UObject**)Elem)
+		{
+			for (int32 Idx = 0; Idx < Num; ++Idx)
+			{
+				UObject** TestObj = (UObject**)Array.GetData(Idx, Layout);
+
+				if (TestObj && *TestObj == *Obj)
+				{
+					Result = Idx;
+					break;
+				}
+			}
+		}
+
+		return Result;
+	}
+
+	inline int32 FindIndexFromNode(const FOverriddenPropertyNodeID& Node, const FOverridableScriptArray& Array) const
+	{
+		int32 Num = Array.Num();
+
+		for (int32 Idx = 0; Idx < Num; ++Idx)
+		{
+			UObject** Item = (UObject**)Array.GetData(Idx, Layout);
+
+			if (Node == FOverriddenPropertyNodeID(*Item))
+			{
+				return Idx;
+			}
+		}
+
+		return INDEX_NONE;
+	}
+
+	inline void SaveOverrides(FMemberBuilder& Dst, const FOverridableScriptArray& Array, const FOverridableScriptArray& Default,
+								const FOverriddenPropertySet& Overrides, 
+								TBitArray<>& RemoveIds, TBitArray<>& InsertIds, TBitArray<>& ModifyIds, bool& bIsFullyReplaced) const
+	{
+		bIsFullyReplaced = false;
+
+		if (const FOverriddenPropertyNode* ArrayNode = GetOverriddenNode(Overrides, &Array))
+		{
+			if (ArrayNode->GetOperation() == EOverriddenPropertyOperation::Replace)
+			{
+				bIsFullyReplaced = true;
+			}
+			else
+			{
+				for (const FOverriddenPropertyNode& SubNode : ArrayNode->GetSubPropertyNodes())
+				{
+					EOverriddenPropertyOperation Op = SubNode.GetOperation();
+
+					switch (Op)
+					{
+					case EOverriddenPropertyOperation::Remove:
+					{
+						const int32 DefaultIdx = FindIndexFromNode(SubNode.GetNodeID(), Default);
+						if (DefaultIdx != INDEX_NONE)
+						{
+							RemoveIds[DefaultIdx] = 1;
+						}
+						else
+						{
+							UE_LOGFMT(LogPlainPropsUObject, VeryVerbose, "Unable to save deleted item '{Item}'.", *SubNode.GetNodeID().ToDebugString());
+						}
+						break;
+					}
+					case EOverriddenPropertyOperation::Add:
+					{
+						const int32 Idx = FindIndexFromNode(SubNode.GetNodeID(), Array);
+						if (Idx != INDEX_NONE)
+						{
+							InsertIds[Idx] = 1;
+						}
+						else
+						{
+							UE_LOGFMT(LogPlainPropsUObject, VeryVerbose, "Unable to save added item '{Item}'.", *SubNode.GetNodeID().ToDebugString());
+						}
+						break;
+					}
+					case EOverriddenPropertyOperation::Modified:
+					{
+						const int32 Idx = FindIndexFromNode(SubNode.GetNodeID(), Array);
+						if (Idx != INDEX_NONE)
+						{
+							ModifyIds[Idx] = 1;
+						}
+						else
+						{
+							UE_LOGFMT(LogPlainPropsUObject, VeryVerbose, "Unable to save modified item '{Item}'.", *SubNode.GetNodeID().ToDebugString());
+						}
+						break;
+					}
+					case EOverriddenPropertyOperation::Externalized:
+					{
+						// Todo: Handle external objects.
+						break;
+					}
+					default:
+						checkf(false, TEXT("Unsupported operation type"));
+						break;
+					}
+				}
+			}
+		}
+	}
+
+	inline void RestoreOverrides(const FOverridableScriptArray& Array, FRangeLoadView Items,
+								 FOverriddenPropertySet* Overrides, EOverriddenPropertyOperation Operation) const
+	{
+		if (Overrides)
+		{
+			if (Operation == EOverriddenPropertyOperation::Replace)
+			{
+				RestoreOverriddenOperation(*Overrides, EOverriddenPropertyOperation::Replace, &Array);
+			}
+			else if (Operation == EOverriddenPropertyOperation::Add || 
+					 Operation == EOverriddenPropertyOperation::Remove ||
+					 Operation == EOverriddenPropertyOperation::Modified)
+			{
+				// We need to add a temporary item to the array which we will remove
+				// after iterating over the items.
+				FOverridableScriptArray& MutArray = const_cast<FOverridableScriptArray&>(Array);
+
+				int32 TmpIdx;
+				void* Tmp = AddItem(MutArray, /* out */ TmpIdx);
+				for (FStructLoadView Item : Items.AsStructs())
+				{
+					LoadStruct(Tmp, Item);
+					int32 Idx = FindIndex(Array, Tmp);
+
+					if (Idx != INDEX_NONE && Idx != TmpIdx)
+					{
+						UObject** SubObject = (UObject**)Array.GetData(Idx, Layout);
+
+						if (*SubObject)
+						{
+							// Need to fetch the ArrayNode every loop as the previous iteration might have reallocated the node.
+							if (FOverriddenPropertyNode* ArrayNode = RestoreOverriddenOperation(*Overrides, EOverriddenPropertyOperation::Modified, &Array))
+							{
+								// Rebuild the overridden info
+								Overrides->RestoreSubObjectOperation(Operation, *ArrayNode, *SubObject);
+							}
+						}
+					}
+				}
+
+				Inner.DestroyItem(Tmp);
+				MutArray.RemoveAtUninitialized(Layout, TmpIdx);
+			}
+			else
+			{
+				checkf(false, TEXT("Unsupported override operation for arrays"));
+			}
+		}
+	}
+
+	inline void DestroyElem(FOverridableScriptArray& Array, int32 Idx) const
+	{
+		if (Idx < Array.Num())
+		{
+			Inner.DestroyItem(Array.GetData(Idx, Layout));
+		}
+	}
+
+	inline void RemoveElem(FOverridableScriptArray& Array, const void* Elem) const
+	{
+		if (int32 Idx = FindIndex(Array, Elem); Idx != INDEX_NONE)
+		{
+			DestroyElem(Array, Idx);
+			Array.Remove(Idx, 1, 8, 16);
+		}
+	}
+
+	inline bool HasItem(const FOverridableScriptArray& Array, const void* Elem) const
+	{
+		return FindIndex(Array, Elem) != INDEX_NONE;
+	}
+
+	inline bool HasKey(const FOverridableScriptArray& Array, const void* Elem) const
+	{
+		return HasItem(Array, Elem);
+	}
+
+	inline void DestroyAll(FOverridableScriptArray& Array) const
+	{
+		for (int32 Idx = 0; Idx < Array.Num(); ++Idx)
+		{
+			void* Item = Array.GetData(Idx, Layout);
+			DestroyValue(Inner.Property, Item);
+		}
+	}
+
+	virtual void ReadItems(FSaveRangeContext& Ctx) const override
+	{
+		const FScriptArray& Array = Ctx.Request.GetRange<FScriptArray>();
+		Ctx.Items.SetAll(Array.GetData(), static_cast<uint64>(Array.Num()), GetStride());
+	}
+
+	virtual void MakeItems(FLoadRangeContext& Ctx) const override
+	{
+		FScriptArray& Array = Ctx.Request.GetRange<FScriptArray>();
+
+		SIZE_T Stride = GetStride();
+
+		int32 NewNum = static_cast<int32>(Ctx.Request.NumTotal());
+		Array.SetNumUninitialized(NewNum, Stride, Layout.AlignmentOfElement);
+		if (NewNum)
+		{
+			FMemory::Memzero(Array.GetData(), NumBytes(NewNum, Stride));
+		}
+
+		Ctx.Items.Set(Array.GetData(), static_cast<uint64>(NewNum), Stride);
+	}
+
+	virtual void SaveCustom(FMemberBuilder& Dst, const void* Src, FBaseline Base, const FSaveContext& Ctx) override
+	{
+		SaveSetOverrides(*this, Dst, *static_cast<const FOverridableScriptArray*>(Src), static_cast<const FOverridableScriptArray*>(Base.Get()),
+						 Ctx, GetSaveOverrides());
+	}
+
+	virtual void LoadCustom(void* Dst, FStructLoadView Src, ECustomLoadMethod Method) const override
+	{
+		check(Method == ECustomLoadMethod::Assign);
+		LoadSetDelta(*this, *static_cast<FOverridableScriptArray*>(Dst), Src, GetRestoreOverrides());
+	}
+
+	virtual bool DiffCustom(const void* A, const void* B, const FBindContext&) const override
+	{
+		return DiffSet(*this, *static_cast<const FOverridableScriptArray*>(A), *static_cast<const FOverridableScriptArray*>(B));
+	}
+
+	// Load into empty array
+	inline void AssignEmpty(FOverridableScriptArray& Dst, FRangeLoadView Src) const
+	{
+		SIZE_T Stride = GetStride();
+		uint8* It = SetNumUninitialized(Dst, Layout, Src.Num());
+		InitStridedItems(Inner, It, Src.Num(), Stride);
+
+		for (FStructLoadView Item : Src.AsStructs())
+		{
+			LoadStruct(It, Item);
+			It += Stride;
+		}
+	}
+
+	inline void* AddItem(FOverridableScriptArray& Dst, int32& OutIdx) const
+	{
+		OutIdx = Dst.AddUninitialized(Layout);
+		void* Out = Dst.GetData(OutIdx, Layout);
+		Inner.InitItem(Out);
+		return Out;
+	}
+
+	// Load structs into non-empty array
+	inline void InsertNonEmpty(FOverridableScriptArray& Dst, FRangeLoadView Src) const
+	{
+		// Written to avoid FProperty::CopyCompleteValue_InContainer dependency
+		// Items are loaded directly into a temporary entry at the end of the array 
+		// and then removed if a duplicate existed.
+
+		int32 TmpIdx;
+		void* Tmp = AddItem(Dst, /* out */ TmpIdx);
+		for (FStructLoadView Item : Src.AsStructs())
+		{
+			LoadStruct(Tmp, Item);
+			int32 FoundIdx = FindIndex(Dst, Tmp);
+
+			if (FoundIdx == TmpIdx)
+			{
+				Tmp = AddItem(Dst, /* out */ TmpIdx);
+			}
+		}
+
+		Inner.DestroyItem(Tmp);
+		Dst.RemoveAtUninitialized(Layout, TmpIdx);
+
+		check(Dst.NumUnchecked() >= Src.Num());
+	}
+
+	inline void Remove(FOverridableScriptArray& Dst, FRangeLoadView Src) const
+	{
+		TArray<uint8, TInlineAllocator<64>> Buffer;
+		Buffer.SetNumUninitialized(Inner.Size);
+		Inner.InitItem(Buffer.GetData());
+		void* Tmp = Buffer.GetData();
+
+		for (FStructLoadView Item : Src.AsStructs())
+		{
+			LoadStruct(Tmp, Item);
+			RemoveElem(Dst, Tmp);
+		}
+
+		Inner.DestroyItem(Tmp);
+	}
+
+	inline void DestroyRemoved(FOverridableScriptArray& Dst, int32 Idx) const
+	{
+		checkSlow(Idx < Dst.GetMaxIndex());
+		checkSlow(!Dst.IsValidIndex(Idx));
+		void* Elem = Dst.GetData(Idx, Layout);
+		Inner.DestroyItem(Elem);
+	}
+};
+
+// Utility that allows us to forward declare FOverridableArrayPropertyBinding.
+inline FOverridableArrayPropertyBinding* CreateAndLeakOverridableArrayBinding(FArrayProperty* Property, FMemberBinding InnerBinding, 
+																			  FBindId BindId, FDeclId DeclId)
+{
+	FOverridableArrayPropertyBinding* New = new FOverridableArrayPropertyBinding(Property, InnerBinding);
+	GUE.Customs.BindStruct(BindId, *New, New->DeclareCustom(DeclId), {});
+	return New;
+}
+
+////////////////////////////////////////////////////////////////////////////////////////////////
+
 // @pre Binding.InnermostId isn't lowered
 static FMemberSpec ToSpec(FMemberBinding Binding)
 {
@@ -3309,7 +4064,18 @@ private:
 	
 	inline FMemberBinding BindArray(FArrayProperty* Property)
 	{
-		return BindAsRange(Property, AllocateArrayBinding(Property), BindSingleProperty(Property->Inner));
+		FMemberBinding Inner = BindSingleProperty(Property->Inner);
+
+		const bool bIsOverridable = Property->HasAnyPropertyFlags(CPF_ExperimentalOverridableLogic);
+
+		if (bIsOverridable)
+		{
+			return BindAsStruct(Property, GCachedArrayBindings.BindOverridableArray(Property, Inner));
+		}
+		else
+		{
+			return BindAsRange(Property, AllocateArrayBinding(Property), Inner);
+		}
 	}
 	inline FMemberBinding BindMap(FMapProperty* Property)
 	{
diff --git a/Engine/Plugins/Experimental/PlainPropsUObject/Source/Private/PlainPropsUObjectDiff.cpp b/Engine/Plugins/Experimental/PlainPropsUObject/Source/Private/PlainPropsUObjectDiff.cpp
--- a/Engine/Plugins/Experimental/PlainPropsUObject/Source/Private/PlainPropsUObjectDiff.cpp
+++ b/Engine/Plugins/Experimental/PlainPropsUObject/Source/Private/PlainPropsUObjectDiff.cpp
@@ -7,6 +7,7 @@
 #include "UObject/Object.h"
 #include "UObject/Package.h"
 #include "UObject/UnrealType.h"
+#include "UObject/OverridableManager.h"
 #include "Containers/StringFwd.h"
 
 // Temp hacks for non-intrusive prototype
@@ -15,6 +16,10 @@ PP_DEFINE_PRIVATE_MEMBER_PTR(FFieldPath, ResolvedOwner,	TWeakObjectPtr<UStruct>)
 PP_DEFINE_PRIVATE_MEMBER_PTR(FFieldPath, Path,			TArray<FName>);
 PP_DEFINE_PRIVATE_MEMBER_NAMESPACE_END()
 
+PP_DEFINE_PRIVATE_MEMBER_NAMESPACE_BEGIN(FOverriddenPropertyNodeIdDiff_Private)
+PP_DEFINE_PRIVATE_MEMBER_PTR(FOverriddenPropertyNodeID, Object, TWeakObjectPtr<const UObject>);
+PP_DEFINE_PRIVATE_MEMBER_NAMESPACE_END()
+
 namespace PlainProps::UE
 {
 
@@ -588,6 +593,100 @@ bool DiffStructs(const void* A, const void* B, const UStruct* Struct, FDiffObjec
 	return false;
 }
 
+static bool EqualOverrideIds(const FOverriddenPropertyNodeID& A, const FOverriddenPropertyNodeID& B)
+{
+	using namespace FOverriddenPropertyNodeIdDiff_Private;
+	const UObject* ObjectA = (A.*_Object).Get();
+	const UObject* ObjectB = (B.*_Object).Get();
+
+	if (ObjectA || ObjectB)
+	{
+		TStringBuilder<256> PathA;
+		if (ObjectA)
+		{
+			GetRelativePathName(ObjectA, PathA);
+		}
+		
+		TStringBuilder<256> PathB;
+		if (ObjectB)
+		{
+			GetRelativePathName(ObjectB, PathB);
+		}
+
+		return PathA.ToView() == PathB.ToView();
+	}
+	else
+	{
+		return A == B;
+	}
+}
+
+static bool DiffOverrideNodes(const FOverriddenPropertyNode& A, const FOverriddenPropertyNode& B)
+{
+	EOverriddenPropertyOperation OpA = A.GetOperation();
+	EOverriddenPropertyOperation OpB = B.GetOperation();
+
+	if ((OpA == OpB) && EqualOverrideIds(A.GetNodeID(), B.GetNodeID()))
+	{
+		const TArray<FOverriddenPropertyNode>& SubNodesA = A.GetSubPropertyNodes();
+		const TArray<FOverriddenPropertyNode>& SubNodesB = B.GetSubPropertyNodes();
+
+		if (SubNodesA.Num() == SubNodesB.Num())
+		{
+			int Index = 0;
+
+			for (const FOverriddenPropertyNode& SubNodeA : SubNodesA)
+			{
+				// The sub-nodes might be in different order so find a match. This
+				// could be optimized, however, the sub-node count is generally low.
+				const FOverriddenPropertyNode* SubNodeB = nullptr;
+				for (const FOverriddenPropertyNode& Test : SubNodesB)
+				{
+					if (EqualOverrideIds(Test.GetNodeID(), SubNodeA.GetNodeID()))
+					{
+						SubNodeB = &Test;
+						break;
+					}
+				}
+
+				if (SubNodeB)
+				{
+					if (DiffOverrideNodes(SubNodeA, *SubNodeB))
+					{
+						return true;
+					}
+				}
+				else
+				{
+					return true;
+				}
+			}
+		}
+		else
+		{
+			return true;
+		}
+	}
+	else
+	{
+		return true;
+	}
+
+	return false;
+}
+
+static bool DiffOverrides(const FOverriddenPropertySet* A, const FOverriddenPropertySet* B)
+{
+	if (A && B)
+	{
+		return DiffOverrideNodes(A->GetRootOverriddenPropertyNode(), B->GetRootOverriddenPropertyNode());
+	}
+	
+	// When both overrides are NULL then the object does not support
+	// overridable serialization and should not be considered a diff.
+	return A != B;
+}
+
 bool DiffObjects(FDiffObjectContext& Ctx)
 {
 	TMap<FString, const UObject*> ObjectsA;
@@ -606,6 +705,8 @@ bool DiffObjects(FDiffObjectContext& Ctx)
 	FillObjectMap(Ctx.RootA, ObjectsA);
 	FillObjectMap(Ctx.RootB, ObjectsB);
 
+	FOverridableManager& OverrideManager = FOverridableManager::Get();
+
 	for (const TPair<FString, const UObject*>& PairA : ObjectsA)
 	{
 		const UObject* A = PairA.Value;
@@ -617,6 +718,12 @@ bool DiffObjects(FDiffObjectContext& Ctx)
 		}
 
 		Ctx.Visited.Reset();
+
+		if (DiffOverrides(OverrideManager.GetOverriddenProperties(A), OverrideManager.GetOverriddenProperties(*B)))
+		{
+			Ctx.OverrideDiffs.Emplace(A, *B);
+		}
+
 		if (DiffObjects(A, *B, Ctx))
 		{
 			return true;
diff --git a/Engine/Plugins/Experimental/PlainPropsUObject/Source/Private/PlainPropsUObjectDiffInternal.h b/Engine/Plugins/Experimental/PlainPropsUObject/Source/Private/PlainPropsUObjectDiffInternal.h
--- a/Engine/Plugins/Experimental/PlainPropsUObject/Source/Private/PlainPropsUObjectDiffInternal.h
+++ b/Engine/Plugins/Experimental/PlainPropsUObject/Source/Private/PlainPropsUObjectDiffInternal.h
@@ -53,6 +53,12 @@ struct FDiffObjectNode
 	}
 };
 
+struct FDiffOverride
+{
+	const UObject* A;
+	const UObject* B;
+};
+
 struct FDiffObjectFilter
 {
 	// Structs for which to bypass the native identical function
@@ -79,6 +85,7 @@ struct FDiffObjectContext
 
 	// out
 	TArray<FDiffObjectNode> Diffs;
+	TArray< FDiffOverride> OverrideDiffs;
 	uint32 NumIgnoredDiffs = 0;
 };
 
diff --git a/Engine/Plugins/Experimental/PlainPropsUObject/Source/Public/PlainPropsUObjectRuntime.h b/Engine/Plugins/Experimental/PlainPropsUObject/Source/Public/PlainPropsUObjectRuntime.h
--- a/Engine/Plugins/Experimental/PlainPropsUObject/Source/Public/PlainPropsUObjectRuntime.h
+++ b/Engine/Plugins/Experimental/PlainPropsUObject/Source/Public/PlainPropsUObjectRuntime.h
@@ -115,6 +115,7 @@ struct FCommonMemberIds
 	FMemberId					Assign;
 	FMemberId					Remove;
 	FMemberId					Insert;
+	FMemberId					Modify;
 	FMemberId					Id;
 	FMemberId					Object;
 	FMemberId					Function;
