PlainProps: Bind FText and FName to the CoreUObject scope so that they can be member types in PP_REFLECT_STRUCT macros.

#rnx
#Synced-CL 52055016

diff --git a/Engine/Plugins/Experimental/PlainPropsUObject/Source/Public/PlainPropsUObjectRuntime.h b/Engine/Plugins/Experimental/PlainPropsUObject/Source/Public/PlainPropsUObjectRuntime.h
--- a/Engine/Plugins/Experimental/PlainPropsUObject/Source/Public/PlainPropsUObjectRuntime.h
+++ b/Engine/Plugins/Experimental/PlainPropsUObject/Source/Public/PlainPropsUObjectRuntime.h
@@ -273,7 +273,9 @@ template<> struct TCustomBind<::UE::FDynamicallyTypedValue> { using Type = UE::F
 template<> struct TCustomBind<FReferencePropertyValue> { using Type = UE::FReferencePropertyBinding; };
 template<> struct TCustomBind<FVerseFunction> { using Type = UE::FVerseFunctionBinding; };
 
-// Temporary way to tie certain types to /Script/CoreUObject scope
+// Temporary way to tie certain Core types to /Script/CoreUObject scope
+template<> inline FScopeId GetOuterScope<UE::FRuntimeIds, TTypename<FName>>()			{ return UE::GUE.Scopes.CoreUObject; }
+template<> inline FScopeId GetOuterScope<UE::FRuntimeIds, TTypename<FText>>()			{ return UE::GUE.Scopes.CoreUObject; }
 template<> inline FScopeId GetOuterScope<UE::FRuntimeIds, TTypename<FIntPoint>>()		{ return UE::GUE.Scopes.CoreUObject; }
 template<> inline FScopeId GetOuterScope<UE::FRuntimeIds, TTypename<FQuat>>()			{ return UE::GUE.Scopes.CoreUObject; }
 template<> inline FScopeId GetOuterScope<UE::FRuntimeIds, TTypename<FTransform>>()		{ return UE::GUE.Scopes.CoreUObject; }
