PlainProps: Fix incorrect include path casing.

#rnx

#Synced-CL 52055016

diff --git a/Engine/Plugins/Experimental/PlainPropsUObject/Source/Private/PlainPropsRoundtripTest.cpp b/Engine/Plugins/Experimental/PlainPropsUObject/Source/Private/PlainPropsRoundtripTest.cpp
--- a/Engine/Plugins/Experimental/PlainPropsUObject/Source/Private/PlainPropsRoundtripTest.cpp
+++ b/Engine/Plugins/Experimental/PlainPropsUObject/Source/Private/PlainPropsRoundtripTest.cpp
@@ -17,7 +17,7 @@
 #include "HAL/FileManager.h"
 #include "JsonObjectGraph/Stringify.h"
 #include "Logging/StructuredLog.h"
-#include "Misc/Filehelper.h"
+#include "Misc/FileHelper.h"
 #include "Misc/PackageName.h"
 #include "Misc/Paths.h"
 #include "Serialization/Archive.h"
