PlainProps match TMap loading to FMapProperty::SerializeItem to workaround potential bug in TCompactSet

#rnx
#Synced-CL 52245732

diff --git a/Engine/Plugins/Experimental/PlainPropsUObject/Source/Private/PlainPropsUObjectBindings.cpp b/Engine/Plugins/Experimental/PlainPropsUObject/Source/Private/PlainPropsUObjectBindings.cpp
--- a/Engine/Plugins/Experimental/PlainPropsUObject/Source/Private/PlainPropsUObjectBindings.cpp
+++ b/Engine/Plugins/Experimental/PlainPropsUObject/Source/Private/PlainPropsUObjectBindings.cpp
@@ -2719,6 +2719,15 @@ struct TMapPropertyCustomBinding : TMapPropertyItemBinding<FScriptMap, KeyKind,
 		return Map.FindPairIndex(Key, Layout, MakeHashFn(InnerKey), MakeIdenticalFn(InnerKey));
 	}
 
+	inline void* FindOrAdd(FScriptMap& Map, const void* Key) const
+	{
+		return Map.FindOrAdd(Key, Layout, MakeHashFn(InnerKey), MakeIdenticalFn(InnerKey), [this](void* K, void* V)
+			{
+				InnerKey.InitItem(K);
+				InnerValue.InitItem(V);
+			});
+	}
+
 	inline bool HasKey(const FScriptMap& Map, const void* Key) const
 	{
 		return FindKey(Map, Key) != INDEX_NONE;
@@ -2820,16 +2829,12 @@ struct TMapPropertyCustomBinding : TMapPropertyItemBinding<FScriptMap, KeyKind,
 
 	// Load structs or ranges into non-empty map
 	inline void InsertNonEmpty(FScriptMap& Dst, FRangeLoadView Src) const
-	{
-		auto HashFn = MakeHashFn(InnerKey);
-
-		// Written to avoid FProperty::CopyCompleteValue_InContainer dependency
-		// Items are loaded directly into sparse array and then removed if a duplicate existed
-
-		const int32 OldNum = Dst.NumUnchecked();
-		const int32 ValueOffset = Layout.ValueOffset;
-		int32 TmpIdx;
-		uint8* Tmp = AddPair(Dst, /* out*/ TmpIdx);
+	{		
+		// Create temporary key
+		TArray<uint8, TInlineAllocator<32>> TmpKey;
+		TmpKey.SetNumUninitialized(Layout.SetLayout.Size);
+		void* Key = TmpKey.GetData();
+		InnerKey.InitItem(Key);
 
 		FOptionalSchemaId InnerLoadIds[2];
 		FStructRangeLoadView Structs = Src.AsStructs();
@@ -2838,31 +2843,12 @@ struct TMapPropertyCustomBinding : TMapPropertyItemBinding<FScriptMap, KeyKind,
 		{
 			// Equivalent to LoadStruct(It, Struct);
 			FBitCacheReader Bits;
-			KeyRange.LoadItem(Tmp,					/* in-out */ Struct.Values, /* in-out */ Bits, InnerLoadIds[0], Struct.Schema.Batch);
-
-			if (int32 Idx = FindKey(Dst, Tmp); Idx != INDEX_NONE)
-			{
-				// Load value into existing pair
-				uint8* Pair = static_cast<uint8*>(Dst.GetData(Idx, Layout));
-				ValueRange.LoadItem(Pair + ValueOffset,	/* in-out */ Struct.Values, /* in-out */ Bits, InnerLoadIds[1], Struct.Schema.Batch);
-			}
-			else
-			{
-				// Load value into tmp pair and add new temporary
-				ValueRange.LoadItem(Tmp + ValueOffset,	/* in-out */ Struct.Values, /* in-out */ Bits, InnerLoadIds[1], Struct.Schema.Batch);
-				Dst.CommitLastUninitialized(Layout, HashFn);
-				Tmp = AddPair(Dst, /* out*/ TmpIdx);
-			}
+			KeyRange.LoadItem(Key,		/* in-out */ Struct.Values, /* in-out */ Bits, InnerLoadIds[0], Struct.Schema.Batch);
+			void* Value = FindOrAdd(Dst, Key);
+			ValueRange.LoadItem(Value,	/* in-out */ Struct.Values, /* in-out */ Bits, InnerLoadIds[1], Struct.Schema.Batch);
 		}
 
-		InnerKey.DestroyItem(Tmp);
-		InnerValue.DestroyItem(Tmp + ValueOffset);
-		Dst.RemoveAtUninitialized(Layout, TmpIdx);
-
-		if (Dst.NumUnchecked() != OldNum)
-		{
-			Dst.CommitAllUninitialized(Layout, HashFn);
-		}
+		InnerKey.DestroyItem(Key);
 	}
 
 	inline void Remove(FScriptMap& Dst, FRangeLoadView Src) const
