PlainProps follow-up on previous TCompactMap load bug

#rnx
#Synced-CL 52262084

diff --git a/Engine/Plugins/Experimental/PlainPropsUObject/Source/Private/PlainPropsUObjectBindings.cpp b/Engine/Plugins/Experimental/PlainPropsUObject/Source/Private/PlainPropsUObjectBindings.cpp
--- a/Engine/Plugins/Experimental/PlainPropsUObject/Source/Private/PlainPropsUObjectBindings.cpp
+++ b/Engine/Plugins/Experimental/PlainPropsUObject/Source/Private/PlainPropsUObjectBindings.cpp
@@ -2587,7 +2587,10 @@ struct TMapPropertyItemBinding : IItemRangeBinding
 	, Layout(In->MapLayout)
 	, InnerKey(In->KeyProp)
 	, InnerValue(In->ValueProp)
-	{}
+	{
+		check(InnerKey.Size <= static_cast<uint64>(Layout.ValueOffset));
+		check(InnerValue.Size <= static_cast<uint64>(Layout.SetLayout.Size - Layout.ValueOffset));
+	}
 
 	inline SIZE_T GetStride() const	{ return Layout.SetLayout.Size; }
 
@@ -2719,15 +2722,6 @@ struct TMapPropertyCustomBinding : TMapPropertyItemBinding<FScriptMap, KeyKind,
 		return Map.FindPairIndex(Key, Layout, MakeHashFn(InnerKey), MakeIdenticalFn(InnerKey));
 	}
 
-	inline void* FindOrAdd(FScriptMap& Map, const void* Key) const
-	{
-		return Map.FindOrAdd(Key, Layout, MakeHashFn(InnerKey), MakeIdenticalFn(InnerKey), [this](void* K, void* V)
-			{
-				InnerKey.InitItem(K);
-				InnerValue.InitItem(V);
-			});
-	}
-
 	inline bool HasKey(const FScriptMap& Map, const void* Key) const
 	{
 		return FindKey(Map, Key) != INDEX_NONE;
@@ -2831,10 +2825,10 @@ struct TMapPropertyCustomBinding : TMapPropertyItemBinding<FScriptMap, KeyKind,
 	inline void InsertNonEmpty(FScriptMap& Dst, FRangeLoadView Src) const
 	{		
 		// Create temporary key
-		TArray<uint8, TInlineAllocator<32>> TmpKey;
-		TmpKey.SetNumUninitialized(Layout.SetLayout.Size);
-		void* Key = TmpKey.GetData();
-		InnerKey.InitItem(Key);
+		TArray<uint8, TInlineAllocator<32>> TmpData;
+		TmpData.SetNumUninitialized(InnerKey.Size);
+		void* TmpKey = TmpData.GetData();
+		InnerKey.InitItem(TmpKey);
 
 		FOptionalSchemaId InnerLoadIds[2];
 		FStructRangeLoadView Structs = Src.AsStructs();
@@ -2843,12 +2837,19 @@ struct TMapPropertyCustomBinding : TMapPropertyItemBinding<FScriptMap, KeyKind,
 		{
 			// Equivalent to LoadStruct(It, Struct);
 			FBitCacheReader Bits;
-			KeyRange.LoadItem(Key,		/* in-out */ Struct.Values, /* in-out */ Bits, InnerLoadIds[0], Struct.Schema.Batch);
-			void* Value = FindOrAdd(Dst, Key);
-			ValueRange.LoadItem(Value,	/* in-out */ Struct.Values, /* in-out */ Bits, InnerLoadIds[1], Struct.Schema.Batch);
+			KeyRange.LoadItem(/* out*/ TmpKey,	/* in-out */ Struct.Values, /* in-out */ Bits, InnerLoadIds[0], Struct.Schema.Batch);
+
+			void* Value = Dst.FindOrAdd(TmpKey, Layout, MakeHashFn(InnerKey), MakeIdenticalFn(InnerKey), [this, TmpKey](void* K, void* V)
+			{
+				FMemory::Memcpy(K, TmpKey, InnerKey.Size);	// Move construct K (assuming trivial relocatability)
+				InnerKey.InitItem(TmpKey);					// Reinit moved from tmp key
+				InnerValue.InitItem(V);						// Default construct V
+			});
+
+			ValueRange.LoadItem(/* out */ Value, /* in-out */ Struct.Values, /* in-out */ Bits, InnerLoadIds[1], Struct.Schema.Batch);
 		}
 
-		InnerKey.DestroyItem(Key);
+		InnerKey.DestroyItem(TmpKey);
 	}
 
 	inline void Remove(FScriptMap& Dst, FRangeLoadView Src) const
