PlainProps: Fix minor mistake in TUEnumAsByteBinding

#rnx

#Synced-CL 52382190

diff --git a/Engine/Plugins/Experimental/PlainPropsUObject/Source/Public/PlainPropsUObjectRuntime.h b/Engine/Plugins/Experimental/PlainPropsUObject/Source/Public/PlainPropsUObjectRuntime.h
--- a/Engine/Plugins/Experimental/PlainPropsUObject/Source/Public/PlainPropsUObjectRuntime.h
+++ b/Engine/Plugins/Experimental/PlainPropsUObject/Source/Public/PlainPropsUObjectRuntime.h
@@ -242,10 +242,9 @@ struct TUEnumAsByteBinding : ICustomBinding
 	const FMemberId				MemberIds[1];
 	const FEnumId				EnumId;
 
-	template<class Ids>
 	TUEnumAsByteBinding(TPropertySpecifier<1>& Spec)
 	: MemberIds{GUE.Members.Value}
-	, EnumId(GetEnumId<Ids, E>())
+	, EnumId(GetEnumId<FRuntimeIds, E>())
 	{
 		static_assert(sizeof(Type) == 1);
 		Spec.Members[0] = FMemberSpec(ELeafType::Enum, ELeafWidth::B8, EnumId);
