PlainProps: Load a blueprint asset completely in the new serialization stack with no Serialize and SerializeScriptProperties calls at all.
Add custom bindings for UStruct, UClass, UFunction and UScriptStruct.
Add custom bindings for FProperty, FBoolProperty, FStructProperty, FClassProperty and FObjectBaseProperty.
Add custom bindings for UEdGraphNode, UEdGraphPin (not-a-UClass) and UK2Node_EditablePinBaseBinding.
Enable usage of TObjectPtr<T> as members in ctti reflected structs

For blueprints, replicate the UEdGraphPin::SerializePinArray functionality required for saving and loading in the new system (ignoring transactions and other archive scenarios).
For data defined types, save and load a unique <class>_C_Meta meta struct schema (modelled with a native super like e.g. UStruct) that describes all FField & FProperty data (names, types, values) required to instantiate the actual class.

Test case with PP_WITH_ENGINE_HOOKS=1 and PP_OVERRIDE_OBJECT_LOADING=1 that depends on https://p4-swarm.epicgames.net/changes/51259023:
.run commandlet TestPlainProps -- -bind=source -batch=linker -load=/Game/GameplayCueNotifies/GCNL_Character_DamageTaken [-text]
PlainPropsUObject Ignored 1 package diff(s) in 1 package(s)
PlainPropsUObject Detected 0 package diff(s) in 1 package(s)
PlainPropsUObject Loaded 182 exports, 1 load plan updates

#rnx

#Synced-CL 52382190

diff --git a/Engine/Plugins/Experimental/PlainProps/Source/Private/PlainPropsSave.cpp b/Engine/Plugins/Experimental/PlainProps/Source/Private/PlainPropsSave.cpp
--- a/Engine/Plugins/Experimental/PlainProps/Source/Private/PlainPropsSave.cpp
+++ b/Engine/Plugins/Experimental/PlainProps/Source/Private/PlainPropsSave.cpp
@@ -436,17 +436,18 @@ static void SaveMemberDelta(FMemberBuilder& Out, const void* Struct, const void*
 	}
 }
 
+const FStructDeclaration& SaveStructDelta(FMemberBuilder& Out, const void* Struct, FBaseline Base, FBindId Id, const FSaveContext& Ctx);
+
 static void SaveSchemaBoundStructDelta(FMemberBuilder& Out, const void* Struct, FBaseline Base, const FSchemaBinding& Schema, const FStructDeclaration& Declaration, const FSaveContext& Ctx)
 {
 	FMemberVisitor It(Schema);
 	if (Declaration.Super)
 	{
 		FBindId SuperId = It.GrabSuper();
-		checkSlow(SuperId == ToOptionalStruct(Declaration.Super));
-		const FStructDeclaration* SuperDecl;
-		const FSchemaBinding& SuperSchema = Ctx.Schemas.GetStruct(SuperId, /* out */ SuperDecl);
-		SaveSchemaBoundStructDelta(Out, Struct, Base.Super(), SuperSchema, *SuperDecl, Ctx);
-		Out.BuildSuperStruct(Ctx.Scratch, *SuperDecl, Ctx.Schemas.GetDebug());
+		// A schema bound struct may have a custom bound super struct
+		check(LowerCast(SuperId) == Declaration.Super);
+		const FStructDeclaration& SuperDecl = SaveStructDelta(Out, Struct, Base.Super(), SuperId, Ctx);
+		Out.BuildSuperStruct(Ctx.Scratch, SuperDecl, Ctx.Schemas.GetDebug());
 	}
 
 	const void* Default = Base.Get();
@@ -523,6 +524,12 @@ FBuiltStruct* SaveStructDeltaIfDiff(const void* Struct, FBaseline Base, FBindId
 	return Out.IsEmpty() ? nullptr : Out.BuildAndReset(Ctx.Scratch, *Declaration, Ctx.Schemas.GetDebug());
 }
 
+void SaveSuper(FMemberBuilder& Out, const void* Super, FBindId SuperId, const FSaveContext& Ctx)
+{
+	const FStructDeclaration& Declaration = SaveStruct(Out, Super, SuperId, Ctx);
+	Out.BuildSuperStruct(Ctx.Scratch, Declaration, Ctx.Schemas.GetDebug());
+}
+
 void SaveSuperDelta(FMemberBuilder& Out, const void* Super, FBaseline SuperBase, FBindId SuperId, const FSaveContext& Ctx)
 {
 	const FStructDeclaration& Declaration = SaveStructDelta(Out, Super, SuperBase, SuperId, Ctx);
diff --git a/Engine/Plugins/Experimental/PlainProps/Source/Public/PlainPropsSave.h b/Engine/Plugins/Experimental/PlainProps/Source/Public/PlainPropsSave.h
--- a/Engine/Plugins/Experimental/PlainProps/Source/Public/PlainPropsSave.h
+++ b/Engine/Plugins/Experimental/PlainProps/Source/Public/PlainPropsSave.h
@@ -60,6 +60,7 @@ private:
 [[nodiscard]] PLAINPROPS_API FBuiltRange*	SaveLeafRange(const void* Range, const ILeafRangeBinding& Binding, FUnpackedLeafType Leaf, const FSaveContext& Ctx);
 [[nodiscard]] PLAINPROPS_API FBuiltRange*	SaveStructRange(TConstArrayView<FBuiltStruct*> Structs, const FSaveContext& Ctx);
 
+PLAINPROPS_API void							SaveSuper(FMemberBuilder& Out, const void* Super, FBindId SuperId, const FSaveContext& Ctx);
 // Save Super if any members differ from SuperBase
 PLAINPROPS_API void							SaveSuperDelta(FMemberBuilder& Out, const void* Super, FBaseline SuperBase, FBindId SuperId, const FSaveContext& Ctx);
 
diff --git a/Engine/Plugins/Experimental/PlainPropsEngine/Source/PlainPropsEngine.Build.cs b/Engine/Plugins/Experimental/PlainPropsEngine/Source/PlainPropsEngine.Build.cs
--- a/Engine/Plugins/Experimental/PlainPropsEngine/Source/PlainPropsEngine.Build.cs
+++ b/Engine/Plugins/Experimental/PlainPropsEngine/Source/PlainPropsEngine.Build.cs
@@ -25,6 +25,7 @@ public class PlainPropsEngine : ModuleRules
 		if (Target.bBuildEditor == true)
 		{
 			PrivateDependencyModuleNames.Add("UnrealEd");
+			PrivateDependencyModuleNames.Add("BlueprintGraph");
 		}
 	}
 }
diff --git a/Engine/Plugins/Experimental/PlainPropsEngine/Source/Private/PlainPropsBlueprintBindings.cpp b/Engine/Plugins/Experimental/PlainPropsEngine/Source/Private/PlainPropsBlueprintBindings.cpp
new file mode 100644
--- /dev/null
+++ b/Engine/Plugins/Experimental/PlainPropsEngine/Source/Private/PlainPropsBlueprintBindings.cpp
@@ -0,0 +1,392 @@
+// Copyright Epic Games, Inc. All Rights Reserved.
+
+#include "PlainPropsEngineBindings.h"
+#include "PlainPropsUObjectRuntime.h"
+#include "PlainPropsUObjectMetaBindingsInternal.h" // Temp for UE_CUSTOM_BIND_MEMBERS and BindStruct()
+#include "PlainPropsRoundtripTest.h" // EBindMode
+#include "PlainPropsBind.h"
+#include "PlainPropsLoad.h"
+#include "PlainPropsSave.h"
+#include "PlainPropsInternalPrivateMemberPtr.h"
+#include "EdGraph/EdGraphNode.h"
+#include "EdGraph/EdGraphPin.h"
+#if WITH_EDITOR
+#include "K2Node_EditablePinBase.h"
+#endif
+
+// Temp hacks for non-intrusive prototype
+PP_DEFINE_PRIVATE_MEMBER_NAMESPACE_BEGIN(EdGraphPin_Private)
+PP_DEFINE_PRIVATE_MEMBER_PTR(FEdGraphPinReference,	OwningNode,	TWeakObjectPtr<UEdGraphNode>);
+PP_DEFINE_PRIVATE_MEMBER_PTR(FEdGraphPinReference,	PinId,		FGuid);
+PP_DEFINE_PRIVATE_MEMBER_NAMESPACE_END()
+
+#if WITH_EDITORONLY_DATA
+PP_REFLECT_STRUCT(, UEdGraphPin, void, PinId, PinName, PinFriendlyName, SourceIndex, PinToolTip, Direction, PinType, DefaultValue, AutogeneratedDefaultValue, DefaultObject, DefaultTextValue, PersistentGuid);
+#else
+PP_REFLECT_STRUCT(, UEdGraphPin, void, PinId, PinName,                  SourceIndex, PinToolTip, Direction, PinType, DefaultValue, AutogeneratedDefaultValue, DefaultObject, DefaultTextValue);
+#endif
+
+
+// PP_REFLECT_ENUM(, EEdGraphPinDirection, EGPD_Input, EGPD_Output, EGPD_MAX);
+
+namespace PlainProps::UE::Blueprint
+{
+
+static FScopeId GetBlueprintScope()
+{
+	return GUE.Scopes.Engine;
+}
+
+// struct id for nested ctti schema binding for UEdGraphPin (not-a-UClass)
+static FDualStructId IndexSchemaEdGraphPin()
+{
+	FType Type =
+	{
+		IndexScope<FRuntimeIds, TTypename<UEdGraphPin>>(),
+		FTypenameId(FRuntimeIds::IndexTypename("UEdGraphPin.Schema"))
+	};
+	FStructId Id = FRuntimeIds::IndexStruct(Type);
+	return FDualStructId(Id);
+}
+
+// struct id for nested UE reflected schema binding for UEdGraphNode (is-a-UClass)
+static FDualStructId IndexSchemaEdGraphNode()
+{
+	FType Type = { GetBlueprintScope(), FTypenameId(FRuntimeIds::IndexTypename("EdGraphNode.Schema")) };
+	FStructId Id = FRuntimeIds::IndexStruct(Type);
+	return FDualStructId(Id);
+}
+
+////////////////////////////////////////////////////////////////////////////////////////////////
+
+struct UEdGraphNodeBinding : ICustomBinding
+{
+	UE_CUSTOM_BIND_MEMBERS(UEdGraphNode, Schema, Pins);
+	UEdGraphNodeBinding(TPropertySpecifier<NumMembers>& Spec);
+	
+	const FDualStructId SchemaId;
+	const FDualStructId EdGraphPinId;
+
+	void Save(FMemberBuilder& Dst, const UEdGraphNode& Src, const UEdGraphNode* Default, const FSaveContext& Ctx) const;
+	void Load(UEdGraphNode& Dst, FStructLoadView Src, ECustomLoadMethod Method) const;
+	static bool Diff(const UEdGraphNode& A, const UEdGraphNode& B, const FBindContext&) { return &A != &B; }
+
+	static FTypedRange SavePinArray(FBindId Id, TConstArrayView<UEdGraphPin*> Pins, const FSaveContext& Ctx);
+	static void LoadPinArray(UEdGraphNode& Dst, FStructRangeLoadView Src);
+
+};
+
+UEdGraphNodeBinding::UEdGraphNodeBinding(TCustomSpecifier<Ids, NumMembers>& Spec)
+: SchemaId(IndexSchemaEdGraphNode())
+, EdGraphPinId(GetStructDualId<Ids, UEdGraphPin>())
+{
+	Spec.Members[M::Schema] = FDeclId(SchemaId);
+	Spec.Members[M::Pins] = DefaultRangeOf(FDeclId(EdGraphPinId));
+}
+
+void UEdGraphNodeBinding::Save(FMemberBuilder& Dst, const UEdGraphNode& Src, const UEdGraphNode* Default, const FSaveContext& Ctx) const
+{
+	// Super::Serialize(Ar);
+	Dst.AddStruct(MemberIds[M::Schema], FBindId(SchemaId), SaveStruct(&Src, FBindId(SchemaId), Ctx));
+	if (Src.Pins.Num())
+	{
+		// see UEdGraphPin::SerializeAsOwningNode()/UEdGraphPin::SerializePinArray()
+		Dst.AddRange(MemberIds[M::Pins], SavePinArray(FBindId(EdGraphPinId), Src.Pins, Ctx));
+	}
+}
+
+void UEdGraphNodeBinding::Load(UEdGraphNode& Dst, FStructLoadView Src, ECustomLoadMethod Method) const
+{
+	FMemberLoader Members(Src);
+	LoadStruct(&Dst, Members.GrabStruct()); // grab schema data
+	if (Members.HasMore())
+	{
+		LoadPinArray(Dst, Members.GrabRange().AsStructs());
+	}
+}
+
+FTypedRange UEdGraphNodeBinding::SavePinArray(FBindId Id, TConstArrayView<UEdGraphPin*> Pins, const FSaveContext& Ctx)
+{
+	FStructRangeSaver Out(Ctx.Scratch, Pins.Num());
+	for (const UEdGraphPin* Pin : Pins)
+	{
+		check(Pin);
+		Out.AddItem(SaveStruct(Pin, Id, Ctx));
+	}
+	return Out.Finalize(MakeStructRangeSchema(DefaultRangeMax, Id));
+}
+
+void UEdGraphNodeBinding::LoadPinArray(UEdGraphNode& Dst, FStructRangeLoadView Src)
+{
+	check(Dst.Pins.IsEmpty());
+	for (FStructLoadView SrcIt : Src)
+	{
+		UEdGraphPin* Pin = UEdGraphPin::CreatePin(&Dst);
+		LoadStruct(Pin, SrcIt);
+		Dst.Pins.Add(Pin);
+#if PP_OVERRIDE_OBJECT_LOADING
+		UEdGraphPin::ResolveReferencesToPin(Pin, /*bStrictValidation*/ true);
+#endif
+	}
+}
+
+////////////////////////////////////////////////////////////////////////////////////////////////
+
+struct UEdGraphPinBinding : ICustomBinding
+{
+#if WITH_EDITORONLY_DATA
+	UE_CUSTOM_BIND_MEMBERS(UEdGraphPin,
+		Schema, ParentPin, ReferencePassThroughConnection, LinkedTo, SubPins,
+		bHidden, bNotConnectable, bDefaultValueIsReadOnly, bDefaultValueIsIgnored,
+		bAdvancedView, bOrphanedPin, bHasSnappedChild, bHasSnappedParent);
+#else
+	UE_CUSTOM_BIND_MEMBERS(UEdGraphPin,
+		Schema, ParentPin, ReferencePassThroughConnection, LinkedTo, SubPins);
+#endif
+	UEdGraphPinBinding(TCustomSpecifier<Ids, NumMembers>& Spec);
+	
+	const FDualStructId SchemaId;
+	const FDualStructId PinReferenceId;
+
+	void Save(FMemberBuilder& Dst, const UEdGraphPin& Src, const UEdGraphPin* Default, const FSaveContext& Ctx) const;
+	void Load(UEdGraphPin& Dst, FStructLoadView Src, ECustomLoadMethod Method) const;
+	static bool Diff(const UEdGraphPin& A, const UEdGraphPin& B, const FBindContext&) { return &A != &B; }
+
+	static FTypedRange SavePinReferenceArray(FBindId Id, TConstArrayView<UEdGraphPin*> Pins, const FSaveContext& Ctx);
+	static void LoadPinReferenceArray(UEdGraphPin& RequestingPin, TArray<UEdGraphPin*>& Pins, FStructRangeLoadView Src, EPinResolveType ResolveType);
+	static void LoadPinReference(UEdGraphPin& RequestingPin, UEdGraphPin*& PinRef, FStructLoadView Struct, EPinResolveType ResolveType, int32 ArrayIdx = INDEX_NONE);
+
+};
+
+UEdGraphPinBinding::UEdGraphPinBinding(TCustomSpecifier<Ids, NumMembers>& Spec)
+: SchemaId(IndexSchemaEdGraphPin())
+, PinReferenceId(GUE.DynamicIds.GetStruct(FEdGraphPinReference::StaticStruct()))
+{
+	Spec.Members[M::Schema] = FDeclId(SchemaId);
+	Spec.Members[M::ParentPin] = FDeclId(PinReferenceId);
+	Spec.Members[M::ReferencePassThroughConnection] = FDeclId(PinReferenceId);
+	Spec.Members[M::LinkedTo] = DefaultRangeOf(FDeclId(PinReferenceId));
+	Spec.Members[M::SubPins] = DefaultRangeOf(FDeclId(PinReferenceId));
+#if WITH_EDITORONLY_DATA
+	Spec.Members[M::bHidden] = Specify<bool>();
+	Spec.Members[M::bNotConnectable] = Specify<bool>();
+	Spec.Members[M::bDefaultValueIsReadOnly] = Specify<bool>();
+	Spec.Members[M::bDefaultValueIsIgnored] = Specify<bool>();
+	Spec.Members[M::bAdvancedView] = Specify<bool>();
+	Spec.Members[M::bOrphanedPin] = Specify<bool>();
+	Spec.Members[M::bHasSnappedChild] = Specify<bool>();
+	Spec.Members[M::bHasSnappedParent] = Specify<bool>();
+#endif
+}
+
+void UEdGraphPinBinding::Save(FMemberBuilder& Dst, const UEdGraphPin& Src, const UEdGraphPin* Default, const FSaveContext& Ctx) const
+{
+#if PP_OVERRIDE_OBJECT_LOADING
+	Dst.AddStruct(MemberIds[M::Schema], FBindId(SchemaId), SaveStruct(&Src, FBindId(SchemaId), Ctx));
+	if (Src.ParentPin)
+	{
+		FEdGraphPinReference PinRef(Src.ParentPin);
+		Dst.AddStruct(MemberIds[M::ParentPin], FBindId(PinReferenceId), SaveStruct(&PinRef, FBindId(PinReferenceId), Ctx));
+	}
+	if (Src.ReferencePassThroughConnection)
+	{
+		FEdGraphPinReference PinRef(Src.ReferencePassThroughConnection);
+		Dst.AddStruct(MemberIds[M::ReferencePassThroughConnection], FBindId(PinReferenceId), SaveStruct(&PinRef, FBindId(PinReferenceId), Ctx));
+	}
+	if (Src.LinkedTo.Num())
+	{
+		// see UEdGraphPin::SerializePinArray()
+		Dst.AddRange(MemberIds[M::LinkedTo], SavePinReferenceArray(FBindId(PinReferenceId), Src.LinkedTo, Ctx));
+	}
+	if (Src.SubPins.Num())
+	{
+		// see UEdGraphPin::SerializePinArray()
+		Dst.AddRange(MemberIds[M::SubPins], SavePinReferenceArray(FBindId(PinReferenceId), Src.SubPins, Ctx));
+	}
+#if WITH_EDITORONLY_DATA
+	Dst.Add(MemberIds[M::bHidden], (bool)Src.bHidden);
+	Dst.Add(MemberIds[M::bNotConnectable], (bool)Src.bNotConnectable);
+	Dst.Add(MemberIds[M::bDefaultValueIsReadOnly], (bool)Src.bDefaultValueIsReadOnly);
+	Dst.Add(MemberIds[M::bDefaultValueIsIgnored], (bool)Src.bDefaultValueIsIgnored);
+	Dst.Add(MemberIds[M::bAdvancedView], (bool)Src.bAdvancedView);
+	Dst.Add(MemberIds[M::bOrphanedPin], (bool)Src.bOrphanedPin);
+	Dst.Add(MemberIds[M::bHasSnappedChild], Src.HasSnappedChild());
+	Dst.Add(MemberIds[M::bHasSnappedParent], Src.HasSnappedParent());
+#endif
+#endif
+}
+
+void UEdGraphPinBinding::Load(UEdGraphPin& Dst, FStructLoadView Src, ECustomLoadMethod Method) const
+{
+#if PP_OVERRIDE_OBJECT_LOADING
+	FMemberLoader Members(Src);
+	LoadStruct(&Dst, Members.GrabStruct()); // grab schema data
+	if (Members.HasMore() && Members.PeekName() == MemberIds[M::ParentPin])
+	{
+		LoadPinReference(Dst, Dst.ParentPin, Members.GrabStruct(), EPinResolveType::ParentPin);
+	}
+	if (Members.HasMore() && Members.PeekName() == MemberIds[M::ReferencePassThroughConnection])
+	{
+		LoadPinReference(Dst, Dst.ReferencePassThroughConnection, Members.GrabStruct(), EPinResolveType::ReferencePassThroughConnection);
+	}
+	if (Members.HasMore() && Members.PeekName() == MemberIds[M::LinkedTo])
+	{
+		LoadPinReferenceArray(Dst, Dst.LinkedTo, Members.GrabRange().AsStructs(), EPinResolveType::LinkedTo);
+	}
+	if (Members.HasMore() && Members.PeekName() == MemberIds[M::SubPins])
+	{
+		LoadPinReferenceArray(Dst, Dst.SubPins, Members.GrabRange().AsStructs(), EPinResolveType::SubPins);
+	}
+#if WITH_EDITORONLY_DATA
+	Dst.bHidden = Members.GrabLeaf().AsBool();
+	Dst.bNotConnectable = Members.GrabLeaf().AsBool();
+	Dst.bDefaultValueIsReadOnly = Members.GrabLeaf().AsBool();
+	Dst.bDefaultValueIsIgnored = Members.GrabLeaf().AsBool();
+	Dst.bAdvancedView = Members.GrabLeaf().AsBool();
+	Dst.bOrphanedPin = Members.GrabLeaf().AsBool();
+	Dst.SetHasSnappedChild(Members.GrabLeaf().AsBool());
+	Dst.SetHasSnappedParent(Members.GrabLeaf().AsBool());
+#endif
+#endif
+}
+
+FTypedRange UEdGraphPinBinding::SavePinReferenceArray(FBindId Id, TConstArrayView<UEdGraphPin*> Pins, const FSaveContext& Ctx)
+{
+	FStructRangeSaver Out(Ctx.Scratch, Pins.Num());
+	for (const UEdGraphPin* Pin : Pins)
+	{
+		FEdGraphPinReference PinRef(Pin);
+		Out.AddItem(SaveStruct(&PinRef, Id, Ctx));
+	}
+	return Out.Finalize(MakeStructRangeSchema(DefaultRangeMax, Id));
+}
+
+void UEdGraphPinBinding::LoadPinReferenceArray(UEdGraphPin& RequestingPin, TArray<UEdGraphPin*>& Pins, FStructRangeLoadView Src, EPinResolveType ResolveType)
+{
+	check(Src.Num());
+	check(Pins.IsEmpty());
+	Pins.SetNum(Src.Num());
+	int32 ArrayIdx = 0;
+	for (FStructLoadView SrcIt : Src)
+	{
+		LoadPinReference(RequestingPin, Pins[ArrayIdx], SrcIt, ResolveType, ArrayIdx);
+		++ArrayIdx;
+	}
+}
+
+void UEdGraphPinBinding::LoadPinReference(UEdGraphPin& RequestingPin, UEdGraphPin*& PinRef, FStructLoadView Struct, EPinResolveType ResolveType, int32 ArrayIdx)
+{
+	using namespace EdGraphPin_Private;
+
+	FEdGraphPinReference Temp;
+	LoadStruct(&Temp, Struct);
+	if (UEdGraphPin* Pin = Temp.Get())
+	{
+		PinRef = Temp.Get();
+	}
+	else
+	{
+		checkf(!PinRef, TEXT("Interesting. The Pin already has a value when the loaded ref resolve is deferred."));
+#if PP_OVERRIDE_OBJECT_LOADING
+		UEdGraphPin::AddUnresolvedPin(RequestingPin, Temp.*_OwningNode, Temp.*_PinId, ResolveType, ArrayIdx);
+#endif
+	}
+}
+
+////////////////////////////////////////////////////////////////////////////////////////////////
+
+#if WITH_EDITOR
+
+struct UK2Node_EditablePinBaseBinding : ICustomBinding
+{
+	UE_CUSTOM_BIND_MEMBERS(UK2Node_EditablePinBase, bIsEditable, UserDefinedPins);
+	UK2Node_EditablePinBaseBinding(TCustomSpecifier<Ids, NumMembers>& Spec);
+
+	const FDualStructId EdGraphNodeId;
+	const FDualStructId UserPinInfoId;
+
+	void Save(FMemberBuilder& Dst, const UK2Node_EditablePinBase& Src, const UK2Node_EditablePinBase* Default, const FSaveContext& Ctx) const;
+	void Load(UK2Node_EditablePinBase& Dst, FStructLoadView Src, ECustomLoadMethod Method) const;
+	static bool Diff(const UK2Node_EditablePinBase& A, const UK2Node_EditablePinBase& B, const FBindContext&) { return &A != &B; }
+};
+
+UK2Node_EditablePinBaseBinding::UK2Node_EditablePinBaseBinding(TCustomSpecifier<Ids, NumMembers>& Spec)
+: EdGraphNodeId(GUE.DynamicIds.GetStruct(UEdGraphNode::StaticClass()))
+, UserPinInfoId(GUE.DynamicIds.GetStruct(FUserPinInfo::StaticStruct()))
+{
+	Spec.Super = FDeclId(EdGraphNodeId);
+	Spec.Members[M::bIsEditable] = Specify<bool>();
+	Spec.Members[M::UserDefinedPins] = DefaultRangeOf(FDeclId(UserPinInfoId));
+}
+
+void UK2Node_EditablePinBaseBinding::Save(FMemberBuilder& Dst, const UK2Node_EditablePinBase& Src, const UK2Node_EditablePinBase* Default, const FSaveContext& Ctx) const
+{
+	SaveSuper(Dst, &Src, FBindId(EdGraphNodeId), Ctx);
+	Dst.Add(MemberIds[M::bIsEditable], (bool)Src.bIsEditable);
+	if (Src.UserDefinedPins.Num())
+	{
+		FStructRangeSaver Out(Ctx.Scratch, Src.UserDefinedPins.Num());
+		for (const TSharedPtr<FUserPinInfo>& PinInfo : Src.UserDefinedPins)
+		{
+			check(PinInfo.IsValid());
+			Out.AddItem(SaveStruct(PinInfo.Get(), FBindId(UserPinInfoId), Ctx));
+		}
+		Dst.AddRange(MemberIds[M::UserDefinedPins], Out.Finalize(MakeStructRangeSchema(DefaultRangeMax, FBindId(UserPinInfoId))));
+	}
+}
+
+void UK2Node_EditablePinBaseBinding::Load(UK2Node_EditablePinBase& Dst, FStructLoadView Src, ECustomLoadMethod Method) const
+{
+	FMemberLoader Members(Src);
+	TryLoadSuper(&Dst, Members);
+	Dst.bIsEditable = Members.GrabLeaf().AsBool();
+	if (Members.HasMore())
+	{
+		for (FStructLoadView SrcIt : Members.GrabRange().AsStructs())
+		{
+			TSharedPtr<FUserPinInfo> PinInfo = MakeShared<FUserPinInfo>();
+			LoadStruct(PinInfo.Get(), SrcIt);
+			Dst.UserDefinedPins.Add(PinInfo);
+		}
+	}
+}
+
+#endif // WITH_EDITOR
+
+} // namespace PlainProps::UE::Blueprint
+
+////////////////////////////////////////////////////////////////////////////////////////////////
+
+namespace PlainProps::UE
+{
+void CustomBindBlueprintTypes(EBindMode Mode)
+{
+#if PP_OVERRIDE_OBJECT_LOADING
+	using namespace Blueprint;
+	// temp: create a custom UEdGraphPin (not-a-UClass) binding with a nested UEdGraphPin.Schema ctti binding
+	new TScopedStructBinding<UEdGraphPin, FDefaultRuntime, void>(IndexSchemaEdGraphPin());
+	new TScopedStructBinding<UEdGraphPin, FDefaultRuntime, UEdGraphPinBinding>();
+	// temp: create a custom UEdGraphNode (is-a-UClass) binding with a nested /Script/Engine.EdGraphNode.Schema reflected property binding
+	BindStruct(FBindId(IndexSchemaEdGraphNode()), UEdGraphNode::StaticClass());
+	new TScopedDefaultUClassBinding<UEdGraphNode, UEdGraphNodeBinding>();
+#if WITH_EDITOR
+	new TScopedDefaultUClassBinding<UK2Node_EditablePinBase, UK2Node_EditablePinBaseBinding>();
+#endif
+#endif
+}
+} // namespace PlainProps::UE
+
+namespace PlainProps
+{
+template<> struct TTypename<EEdGraphPinDirection>
+{ inline static constexpr std::string_view DeclName = "EEdGraphPinDirection"; };
+template<> inline FScopeId GetOuterScope<UE::FRuntimeIds, TTypename<EEdGraphPinDirection>>()
+{ return UE::Blueprint::GetBlueprintScope(); }
+
+// Enable usage of FEdGraphPinType as a member in ctti reflected structs like UEdGraphPin 
+template<> struct TTypename<FEdGraphPinType>
+{ inline static constexpr std::string_view DeclName = "EdGraphPinType"; };
+template<> inline FScopeId GetOuterScope<UE::FRuntimeIds, TTypename<FEdGraphPinType>>()
+{ return UE::Blueprint::GetBlueprintScope(); }
+} // namespace PlainProps
+
diff --git a/Engine/Plugins/Experimental/PlainPropsEngine/Source/Private/PlainPropsCommandlets.cpp b/Engine/Plugins/Experimental/PlainPropsEngine/Source/Private/PlainPropsCommandlets.cpp
--- a/Engine/Plugins/Experimental/PlainPropsEngine/Source/Private/PlainPropsCommandlets.cpp
+++ b/Engine/Plugins/Experimental/PlainPropsEngine/Source/Private/PlainPropsCommandlets.cpp
@@ -312,6 +312,7 @@ int32 UTestPlainPropsCommandlet::Main(const FString& Params)
 		CustomBindEngineTypes(Mode);
 		CustomBindAnimationTypes(Mode);
 		CustomBindNiagaraTypes(Mode);
+		CustomBindBlueprintTypes(Mode);
 	}
 
 	// Default roundtrip options
diff --git a/Engine/Plugins/Experimental/PlainPropsEngine/Source/Private/PlainPropsEngineBindings.h b/Engine/Plugins/Experimental/PlainPropsEngine/Source/Private/PlainPropsEngineBindings.h
--- a/Engine/Plugins/Experimental/PlainPropsEngine/Source/Private/PlainPropsEngineBindings.h
+++ b/Engine/Plugins/Experimental/PlainPropsEngine/Source/Private/PlainPropsEngineBindings.h
@@ -12,5 +12,6 @@ namespace PlainProps::UE
 	void CustomBindEngineTypes(EBindMode Mode);
 	void CustomBindAnimationTypes(EBindMode Mode);
 	void CustomBindNiagaraTypes(EBindMode Mode);
+	void CustomBindBlueprintTypes(EBindMode Mode);
 
 } // namespace PlainProps::UE
diff --git a/Engine/Plugins/Experimental/PlainPropsUObject/Source/Internal/PlainPropsUObjectMetaBindingsInternal.h b/Engine/Plugins/Experimental/PlainPropsUObject/Source/Internal/PlainPropsUObjectMetaBindingsInternal.h
new file mode 100644
--- /dev/null
+++ b/Engine/Plugins/Experimental/PlainPropsUObject/Source/Internal/PlainPropsUObjectMetaBindingsInternal.h
@@ -0,0 +1,344 @@
+// Copyright Epic Games, Inc. All Rights Reserved.
+
+#pragma once
+
+#include "PlainPropsUObjectRuntime.h"
+#include "UObject/Class.h"
+
+////////////////////////////////////////////////////////////////////////////////////////////////
+// Temporary private UE_CUSTOM_BIND_MEMBERS helpers
+
+#define _PP_MEMBER_ID(N, NS, T, M) Ids::IndexMember(#M),
+
+#define _PP_CUSTOM_BIND_MEMBERS(N, ...) \
+	struct M { enum : uint16 { __VA_ARGS__ }; }; \
+	inline static constexpr int NumMembers = N; \
+	const TStaticArray<FMemberId, NumMembers> MemberIds = { _PP_EXPAND(_PP_REFLECT_THINGS(N, _PP_MEMBER_ID, _, _, __VA_ARGS__)) };
+
+// Temporary UE_CUSTOM_BIND_MEMBERS macro that eliminates the manual hassle of 
+// of sizing arrays, indexing names and keeping track of indices when modifying custom bindings.
+// Todo: Handle Name-Spec pairs and interleaving of custom members with schema members.
+#define UE_CUSTOM_BIND_MEMBERS(T, ...) \
+	using Type = T; \
+	using Ids = FRuntimeIds; \
+	_PP_CUSTOM_BIND_MEMBERS(_PP_NUM_ARGS(__VA_ARGS__), __VA_ARGS__) \
+
+////////////////////////////////////////////////////////////////////////////////////////////////
+
+// Temporary macro to switch between property override loading and object override loading
+#define PP_OVERRIDE_OBJECT_LOADING 0
+
+PP_REFLECT_STRUCT(, FImplementedInterface, void, Class, PointerOffset, bImplementedByK2);
+
+PP_REFLECT_FLAG_ENUM(, EClassFlags, 
+	CLASS_None,
+	CLASS_Abstract,
+	CLASS_DefaultConfig,
+	CLASS_Config,
+	CLASS_Transient,
+	CLASS_Optional,
+	CLASS_MatchedSerializers,
+	CLASS_ProjectUserConfig,
+	CLASS_Native,
+	CLASS_NotPlaceable,
+	CLASS_PerObjectConfig,
+	CLASS_ReplicationDataIsSetUp,
+	CLASS_EditInlineNew,
+	CLASS_CollapseCategories,
+	CLASS_Interface,
+	CLASS_PerPlatformConfig,
+	CLASS_Const,
+	CLASS_NeedsDeferredDependencyLoading,
+	CLASS_CompiledFromBlueprint,
+	CLASS_MinimalAPI,
+	CLASS_RequiredAPI,
+	CLASS_DefaultToInstanced,
+	CLASS_TokenStreamAssembled,
+	CLASS_HasInstancedReference,
+	CLASS_Hidden,
+	CLASS_Deprecated,
+	CLASS_HideDropDown,
+	CLASS_GlobalUserConfig,
+	CLASS_Intrinsic,
+	CLASS_Constructed,
+	CLASS_ConfigDoNotCheckDefaults,
+	CLASS_NewerVersionExists
+);
+
+PP_REFLECT_FLAG_ENUM(, EFunctionFlags, 
+	FUNC_None,
+	FUNC_Final,
+	FUNC_RequiredAPI,
+	FUNC_BlueprintAuthorityOnly,
+	FUNC_BlueprintCosmetic,
+	FUNC_Net,
+	FUNC_NetReliable,
+	FUNC_NetRequest,
+	FUNC_Exec,
+	FUNC_Native,
+	FUNC_Event,
+	FUNC_NetResponse,
+	FUNC_Static,
+	FUNC_NetMulticast,
+	FUNC_UbergraphFunction,
+	FUNC_MulticastDelegate,
+	FUNC_Public,
+	FUNC_Private,
+	FUNC_Protected,
+	FUNC_Delegate,
+	FUNC_NetServer,
+	FUNC_HasOutParms,
+	FUNC_HasDefaults,
+	FUNC_NetClient,
+	FUNC_DLLImport,
+	FUNC_BlueprintCallable,
+	FUNC_BlueprintEvent,
+	FUNC_BlueprintPure,
+	FUNC_EditorOnly,
+	FUNC_Const,
+	FUNC_NetValidate
+);
+
+PP_REFLECT_FLAG_ENUM(, EStructFlags, 
+	STRUCT_NoFlags,
+	STRUCT_Native,
+	STRUCT_IdenticalNative,
+	STRUCT_HasInstancedReference,
+	STRUCT_NoExport,
+	STRUCT_Atomic,
+	STRUCT_Immutable,
+	STRUCT_AddStructReferencedObjects,
+	STRUCT_RequiredAPI,
+	STRUCT_NetSerializeNative,
+	STRUCT_SerializeNative,
+	STRUCT_CopyNative,
+	STRUCT_IsPlainOldData,
+	STRUCT_NoDestructor,
+	STRUCT_ZeroConstructor,
+	STRUCT_ExportTextItemNative,
+	STRUCT_ImportTextItemNative,
+	STRUCT_PostSerializeNative,
+	STRUCT_SerializeFromMismatchedTag,
+	STRUCT_NetDeltaSerializeNative,
+	STRUCT_PostScriptConstruct,
+	STRUCT_NetSharedSerialization,
+	STRUCT_Trashed,
+	STRUCT_NewerVersionExists,
+	STRUCT_CanEditChange,
+	STRUCT_Visitor,
+	STRUCT_PostLoad
+);
+
+PP_REFLECT_FLAG_ENUM(, EPropertyFlags, 
+	CPF_None,
+	CPF_Edit,
+	CPF_ConstParm,
+	CPF_BlueprintVisible,
+	CPF_ExportObject,
+	CPF_BlueprintReadOnly,
+	CPF_Net,
+	CPF_EditFixedSize,
+	CPF_Parm,
+	CPF_OutParm,
+	CPF_ZeroConstructor,
+	CPF_ReturnParm,
+	CPF_DisableEditOnTemplate,
+	CPF_NonNullable,
+	CPF_Transient,
+	CPF_Config,
+	CPF_RequiredParm,
+	CPF_DisableEditOnInstance,
+	CPF_EditConst,
+	CPF_GlobalConfig,
+	CPF_InstancedReference,
+	CPF_ExperimentalExternalObjects,
+	CPF_DuplicateTransient,
+	CPF_SaveGame,
+	CPF_NoClear,
+	CPF_Virtual,
+	CPF_ReferenceParm,
+	CPF_BlueprintAssignable,
+	CPF_Deprecated,
+	CPF_IsPlainOldData,
+	CPF_RepSkip,
+	CPF_RepNotify,
+	CPF_Interp,
+	CPF_NonTransactional,
+	CPF_EditorOnly,
+	CPF_NoDestructor,
+	CPF_AutoWeak,
+	CPF_ContainsInstancedReference,
+	CPF_AssetRegistrySearchable,
+	CPF_SimpleDisplay,
+	CPF_AdvancedDisplay,
+	CPF_Protected,
+	CPF_BlueprintCallable,
+	CPF_BlueprintAuthorityOnly,
+	CPF_TextExportTransient,
+	CPF_NonPIEDuplicateTransient,
+	CPF_ExposeOnSpawn,
+	CPF_PersistentInstance,
+	CPF_UObjectWrapper,
+	CPF_HasGetValueTypeHash,
+	CPF_NativeAccessSpecifierPublic,
+	CPF_NativeAccessSpecifierProtected,
+	CPF_NativeAccessSpecifierPrivate,
+	CPF_SkipSerialization,
+	CPF_TObjectPtr,
+	CPF_ExperimentalOverridableLogic,
+	CPF_ExperimentalAlwaysOverriden,
+	CPF_ExperimentalNeverOverriden,
+	CPF_AllowSelfReference,
+	CPF_ForcePostConstructLink
+);
+
+namespace PlainProps::UE
+{
+
+// Todo: Move back to PlainPropsUObjectBindingsInternal.h,
+// where it is used to bind dynamic structs as they are being loaded.
+// It is temporarily exposed here for allowing duplication of a UStruct binding with a different bind id.
+PLAINPROPSUOBJECT_API void BindStruct(FBindId Id, const UStruct* Struct);
+
+////////////////////////////////////////////////////////////////////////////////////////////////
+// Native FProperty meta bindings
+
+struct FPropertyBinding : ICustomBinding
+{
+	UE_CUSTOM_BIND_MEMBERS(FProperty,
+		NamePrivate, ArrayDim, ElementSize, PropertyFlags, RepIndex, BlueprintReplicationCondition);
+	FPropertyBinding(TCustomSpecifier<Ids, NumMembers>& Spec);
+
+	const FEnumId PropertyFlagsId;
+	const FEnumId LifetimeConditionId;
+
+	void		Save(FMemberBuilder& Dst, const Type& Src, const Type* Default, const FSaveContext& Ctx) const;
+	void		Load(Type& Dst, FStructLoadView Src, ECustomLoadMethod Method) const;
+	static bool	Diff(const Type& A, const Type& B, const FBindContext& Ctx) { return &A != &B; }
+};
+
+struct FBoolPropertyBinding : ICustomBinding
+{
+	UE_CUSTOM_BIND_MEMBERS(FBoolProperty, FieldSize, ByteOffset, ByteMask, FieldMask);
+	FBoolPropertyBinding(TCustomSpecifier<Ids, NumMembers>& Spec);
+
+	void		Save(FMemberBuilder& Dst, const Type& Src, const Type* Default, const FSaveContext& Ctx) const;
+	void		Load(Type& Dst, FStructLoadView Src, ECustomLoadMethod Method) const;
+	static bool	Diff(const Type& A, const Type& B, const FBindContext& Ctx) { return &A != &B; }
+};
+
+struct FStructPropertyBinding : ICustomBinding
+{
+	UE_CUSTOM_BIND_MEMBERS(FStructProperty, Struct);
+	FStructPropertyBinding(TCustomSpecifier<Ids, NumMembers>& Spec);
+
+	void		Save(FMemberBuilder& Dst, const Type& Src, const Type* Default, const FSaveContext& Ctx) const;
+	void		Load(Type& Dst, FStructLoadView Src, ECustomLoadMethod Method) const;
+	static bool	Diff(const Type& A, const Type& B, const FBindContext& Ctx) { return &A != &B; }
+};
+
+struct FClassPropertyBinding : ICustomBinding
+{
+	UE_CUSTOM_BIND_MEMBERS(FClassProperty, MetaClass);
+	FClassPropertyBinding(TCustomSpecifier<Ids, NumMembers>& Spec);
+
+	void		Save(FMemberBuilder& Dst, const Type& Src, const Type* Default, const FSaveContext& Ctx) const;
+	void		Load(Type& Dst, FStructLoadView Src, ECustomLoadMethod Method) const;
+	static bool	Diff(const Type& A, const Type& B, const FBindContext& Ctx) { return &A != &B; }
+};
+
+struct FObjectPropertyBaseBinding : ICustomBinding
+{
+	UE_CUSTOM_BIND_MEMBERS(FObjectPropertyBase, PropertyClass);
+	FObjectPropertyBaseBinding(TCustomSpecifier<Ids, NumMembers>& Spec);
+
+	void		Save(FMemberBuilder& Dst, const Type& Src, const Type* Default, const FSaveContext& Ctx) const;
+	void		Load(Type& Dst, FStructLoadView Src, ECustomLoadMethod Method) const;
+	static bool	Diff(const Type& A, const Type& B, const FBindContext& Ctx) { return &A != &B; }
+};
+
+////////////////////////////////////////////////////////////////////////////////////////////////
+// Native UField meta bindings
+
+struct UStructBinding : ICustomBinding
+{
+	UE_CUSTOM_BIND_MEMBERS(UStruct, SuperStruct, Children, ChildPropertyClasses);
+	UStructBinding(TCustomSpecifier<Ids, NumMembers>& Spec);
+
+	static FTypedRange SaveChildrenStatic(const UField* Children, int32 Num, const FSaveContext& Ctx);
+	static FTypedRange SaveChildPropertyClassesStatic(const FField* ChildProperties, int32 Num, const FSaveContext& Ctx);
+	static void LoadChildrenStatic(TObjectPtr<UField>& Children, FStructRangeLoadView Src);
+	static void LoadChildPropertyClassesStatic(UStruct& Dst, FStructRangeLoadView Src);
+
+	void		Save(FMemberBuilder& Dst, const Type& Src, const Type* Default, const FSaveContext& Ctx) const;
+	void		Load(Type& Dst, FStructLoadView Src, ECustomLoadMethod Method) const;
+	static bool	Diff(const Type& A, const Type& B, const FBindContext& Ctx) { return &A != &B; }
+};
+
+struct UClassBinding : ICustomBinding
+{
+#if WITH_EDITORONLY_DATA
+	UE_CUSTOM_BIND_MEMBERS(UClass, ClassFlags, ClassWithin, ClassConfigName, ClassGeneratedBy, Interfaces);
+#else
+	UE_CUSTOM_BIND_MEMBERS(UClass, ClassFlags, ClassWithin, ClassConfigName, Interfaces);
+#endif
+	UClassBinding(TCustomSpecifier<Ids, NumMembers>& Spec);
+
+	const FEnumId ClassFlagsId;
+	const FDualStructId ImplementedInterfaceId;
+
+	static FTypedRange SaveInterfacesStatic(FBindId Id, TConstArrayView<FImplementedInterface> Interfaces, const FSaveContext& Ctx);
+	static void LoadInterfacesStatic(TArray<FImplementedInterface>& Interfaces, FStructRangeLoadView Src);
+
+	void		Save(FMemberBuilder& Dst, const Type& Src, const Type* Default, const FSaveContext& Ctx) const;
+	void		Load(Type& Dst, FStructLoadView Src, ECustomLoadMethod Method) const;
+	static bool	Diff(const Type& A, const Type& B, const FBindContext& Ctx) { return &A != &B; }
+};
+
+struct UFunctionBinding : ICustomBinding
+{
+	UE_CUSTOM_BIND_MEMBERS(UFunction, FunctionFlags, EventGraphFunction, EventGraphCallOffset);
+	UFunctionBinding(TCustomSpecifier<Ids, NumMembers>& Spec);
+
+	const FEnumId FunctionFlagsId;
+
+	void		Save(FMemberBuilder& Dst, const Type& Src, const Type* Default, const FSaveContext& Ctx) const;
+	void		Load(Type& Dst, FStructLoadView Src, ECustomLoadMethod Method) const;
+	static bool	Diff(const Type& A, const Type& B, const FBindContext& Ctx) { return &A != &B; }
+};
+
+struct UScriptStructBinding : ICustomBinding
+{
+	UE_CUSTOM_BIND_MEMBERS(UScriptStruct, StructFlags);
+	UScriptStructBinding(TCustomSpecifier<Ids, NumMembers>& Spec);
+
+	const FEnumId StructFlagsId;
+
+	void		Save(FMemberBuilder& Dst, const Type& Src, const Type* Default, const FSaveContext& Ctx) const;
+	void		Load(Type& Dst, FStructLoadView Src, ECustomLoadMethod Method) const;
+	static bool	Diff(const Type& A, const Type& B, const FBindContext& Ctx) { return &A != &B; }
+};
+
+////////////////////////////////////////////////////////////////////////////////////////////////
+
+} // namespace PlainProps::UE
+
+namespace PlainProps
+{
+// custom FProperty meta bindings
+template<> struct TCustomBind<FProperty> { using Type = UE::FPropertyBinding; };
+template<> struct TCustomBind<FBoolProperty> { using Type = UE::FBoolPropertyBinding; };
+template<> struct TCustomBind<FStructProperty> { using Type = UE::FStructPropertyBinding; };
+template<> struct TCustomBind<FClassProperty> { using Type = UE::FClassPropertyBinding; };
+template<> struct TCustomBind<FObjectPropertyBase> { using Type = UE::FObjectPropertyBaseBinding; };
+
+// custom UField meta bindings
+template<> struct TCustomBind<UStruct> { using Type = UE::UStructBinding; };
+template<> struct TCustomBind<UClass> { using Type = UE::UClassBinding; };
+template<> struct TCustomBind<UFunction> { using Type = UE::UFunctionBinding; };
+template<> struct TCustomBind<UScriptStruct> { using Type = UE::UScriptStructBinding; };
+
+// traits for aux bindings
+template<> struct TOccupancyOf<FImplementedInterface> : FRequireAll {};
+template<> inline FScopeId GetOuterScope<UE::FRuntimeIds, TTypename<FImplementedInterface>>() { return UE::GUE.Scopes.CoreUObject; }
+}
diff --git a/Engine/Plugins/Experimental/PlainPropsUObject/Source/Private/PlainPropsRoundtripTest.cpp b/Engine/Plugins/Experimental/PlainPropsUObject/Source/Private/PlainPropsRoundtripTest.cpp
--- a/Engine/Plugins/Experimental/PlainPropsUObject/Source/Private/PlainPropsRoundtripTest.cpp
+++ b/Engine/Plugins/Experimental/PlainPropsUObject/Source/Private/PlainPropsRoundtripTest.cpp
@@ -2,6 +2,7 @@
 
 #include "PlainPropsRoundtripTest.h"
 #include "PlainPropsUObjectBindingsInternal.h"
+#include "PlainPropsUObjectMetaBindingsInternal.h"
 #include "PlainPropsUObjectDiffInternal.h"
 #include "PlainPropsBuildSchema.h"
 #include "PlainPropsLoadMember.h"
@@ -283,6 +284,8 @@ public:
 			if (SuperIt == ArchClass)
 			{
 				Chain.Emplace(Arch);
+				check(Object->HasAllFlags(RF_ClassDefaultObject));
+				check(!Object->GetClass()->HasAllClassFlags(CLASS_Native));
 				return SaveStructDelta(Object, FBaseline(Chain.GetData(), Chain.Num()), Id, DeltaCtx);
 			}
 		}
@@ -311,6 +314,13 @@ public:
 		SavedObjects.Emplace(Id, Built, Object);
 	}
 
+	void SaveMeta(FBindId Id, const UStruct* Struct)
+	{
+		FBuiltStruct* Built = SaveStruct(Struct, Id, FlatCtx);
+		ObjectMap.Add(SavedObjects.Num());
+		SavedObjects.Emplace(Id, Built, Struct);
+	}
+
 	void SaveObjectWithOverrides(FBindId Id, const UObject* Object, 
 								 const FOverriddenPropertySet& Overriddes, const UObject* Arch)
 	{
@@ -1453,22 +1463,39 @@ struct FPlainPropsLinkerSandbox
 		FBatchSaver BatchSaver(Customs.Overlay, LinkerSave.ExportMap.Num());
 		for (const FObjectExport& Export : LinkerSave.ExportMap)
 		{
-			const UStruct* Struct = SkipEmptyBases(Export.Object->GetClass());
-			FBindId Id = GUE.Names.IndexBindId(PlainProps::UE::IndexStruct(Struct));
-			if (GUE.Schemas.FindStruct(Id))
+			const UStruct* Struct;
+#if PP_OVERRIDE_OBJECT_LOADING
+			if ((Struct = Cast<UStruct>(Export.Object)))
 			{
-				if (const FOverriddenPropertySet* Overriddes = FOverridableManager::Get().GetOverriddenProperties(Export.Object))
+				// Save data defined UStructs as meta structs (_C_META) with no archetype
+				check(!FOverridableManager::Get().GetOverriddenProperties(Export.Object));
+				FBindId Id = GUE.Names.IndexBindId(PlainProps::UE::IndexStructMeta(Struct));
+				if (GUE.Customs.FindStruct(Id))
 				{
-					BatchSaver.SaveObjectWithOverrides(Id, Export.Object, *Overriddes, Export.Object->GetArchetype());
+					BatchSaver.SaveMeta(Id, Struct);
 				}
 				else
 				{
-					BatchSaver.Save(Id, Export.Object, Export.Object->GetArchetype());
+					BatchSaver.Skip(Id, Export.Object);
 				}
 			}
 			else
+#endif
 			{
-				BatchSaver.Skip(Id, Export.Object);
+				Struct = SkipEmptyBases(Export.Object->GetClass());
+				FBindId Id = GUE.Names.IndexBindId(PlainProps::UE::IndexStruct(Struct));
+				if (const FOverriddenPropertySet* Overriddes = FOverridableManager::Get().GetOverriddenProperties(Export.Object))
+				{
+					BatchSaver.SaveObjectWithOverrides(Id, Export.Object, *Overriddes, Export.Object->GetArchetype());
+				}
+				else if (GUE.Schemas.FindStruct(Id))
+				{
+					BatchSaver.Save(Id, Export.Object, Export.Object->GetArchetype());
+				}
+				else
+				{
+					BatchSaver.Skip(Id, Export.Object);
+				}
 			}
 		}
 		TArray64<uint8> StableData = BatchSaver.Write();
@@ -1584,15 +1611,13 @@ struct FPlainPropsLinkerSandbox
 		}
 	}
 
-	static bool LoadObject(FStructuredArchive::FSlot& Slot, uint8* Data)
+	static bool LoadObject(UObject* Object, FArchive& Ar)
 	{
-		FArchive& Ar = Slot.GetUnderlyingArchive();
 		if (!Ar.IsLoading() || !Ar.IsPersistent() || Ar.HasAnyPortFlags(PPF_Duplicate|PPF_DuplicateForPIE))
 		{
 			return false;
 		}
 
-		UObject* Object = (UObject*)Data;
 		FName PackageName = Object->GetPackage()->GetFName();
 		TUniquePtr<FMountedLinkerBatch>* LinkerBatchPtr = MountedPlainpropsPackages.Find(PackageName);
 		check(LinkerBatchPtr);
@@ -1659,17 +1684,31 @@ struct FPlainPropsLinkerSandbox
 				BatchLoader.LoadObject(Object, ExportIndex);
 			}
 		}
-
+		// We only need to adjust the archive when we override property loading,
+		// and not when completely overriding custom serialize object loading.
+		if (GSerializePropertiesOverride && !GSerializeObjectOverride)
+		{
 #if WITH_EDITOR
-		const int32 Pos = Ar.Tell();
-		check(Pos == SerializeContext->SerializedObjectScriptStartOffset);
-		const int32 ScriptSerializationSize = Export.ScriptSerializationEndOffset - Export.ScriptSerializationStartOffset;
-		check(ScriptSerializationSize >= 0);
-		Ar.Seek(Pos + ScriptSerializationSize);
+			const int32 Pos = Ar.Tell();
+			check(Pos == SerializeContext->SerializedObjectScriptStartOffset);
+			const int32 ScriptSerializationSize = Export.ScriptSerializationEndOffset - Export.ScriptSerializationStartOffset;
+			check(ScriptSerializationSize >= 0);
+			Ar.Seek(Pos + ScriptSerializationSize);
 #else
-		// todo: SerializedObjectScriptStartOffset does not exist in client builds
-		unimplemented();
+			// todo: SerializedObjectScriptStartOffset does not exist in client builds
+			unimplemented();
 #endif
+		}
+		return true;
+	}
+
+	static bool LoadProperties(FStructuredArchive::FSlot& Slot, uint8* Data)
+	{
+		FArchive& Ar = Slot.GetUnderlyingArchive();
+		if (!LoadObject((UObject*)Data, Ar))
+		{
+			return false;
+		}
 
 		return true;
 	}
@@ -2052,12 +2091,17 @@ struct FScopeLoadTestPackage
 {
 	FScopeLoadTestPackage(bool bInPlainProps)
 		: SerializePropertiesFunc(GSerializePropertiesOverride)
+		, SerializeObjectFunc(GSerializeObjectOverride)
 		, LoadFromPackageFileFunc(GLoadFromPackageFileOverride)
 		, bPlainProps(bInPlainProps)
 	{
 		if (bPlainProps)
 		{
-			GSerializePropertiesOverride = &FPlainPropsLinkerSandbox::LoadObject;
+#if PP_OVERRIDE_OBJECT_LOADING
+			GSerializeObjectOverride = &FPlainPropsLinkerSandbox::LoadObject;
+#else
+			GSerializePropertiesOverride = &FPlainPropsLinkerSandbox::LoadProperties;
+#endif
 		}
 		GLoadFromPackageFileOverride = &FPlainPropsLinkerSandbox::OpenPackageReader;
 	}
@@ -2067,11 +2111,13 @@ struct FScopeLoadTestPackage
 		if (bPlainProps)
 		{
 			GSerializePropertiesOverride = SerializePropertiesFunc;
+			GSerializeObjectOverride = SerializeObjectFunc;
 		}
 		GLoadFromPackageFileOverride = LoadFromPackageFileFunc;
 	}
 
-	bool(*SerializePropertiesFunc)(class FStructuredArchiveSlot&, uint8*);
+	bool (*SerializePropertiesFunc)(FStructuredArchiveSlot&, uint8*) = nullptr;
+	bool (*SerializeObjectFunc)(UObject*, FArchive&) = nullptr;
 	bool (*LoadFromPackageFileFunc)(const FPackagePath&, FArchive*&) = nullptr;
 	bool bPlainProps;
 };
diff --git a/Engine/Plugins/Experimental/PlainPropsUObject/Source/Private/PlainPropsUObjectBindings.cpp b/Engine/Plugins/Experimental/PlainPropsUObject/Source/Private/PlainPropsUObjectBindings.cpp
--- a/Engine/Plugins/Experimental/PlainPropsUObject/Source/Private/PlainPropsUObjectBindings.cpp
+++ b/Engine/Plugins/Experimental/PlainPropsUObject/Source/Private/PlainPropsUObjectBindings.cpp
@@ -2,8 +2,9 @@
 
 #include "PlainPropsUObjectBindingsInternal.h"
 #include "PlainPropsRoundtripTest.h"
-#include "PlainPropsUObjectRuntime.h"
+#include "PlainPropsUObjectMetaBindingsInternal.h"
 #include "PlainPropsUObjectModule.h"
+#include "PlainPropsUObjectRuntime.h"
 #include "PlainPropsUeCoreBindings.h"
 #include "PlainPropsVisualize.h"
 #include "PlainPropsRestoreOverridableInternal.h"
@@ -304,11 +305,28 @@ FCommonEnumIds::FCommonEnumIds(const FCommonScopeIds& Scopes, TIdIndexer<FSensit
 : PixelFormat(Names.IndexEnum({Scopes.CoreUObject, Names.MakeTypename("PixelFormat")}))
 {}
 
+FCommonPropertyStructIds::FCommonPropertyStructIds(const FCommonScopeIds& Scopes, TIdIndexer<FSensitiveName>& Names)
+: Property(					Names.IndexStruct({Scopes.CoreUObject, Names.MakeTypename("FProperty")}))
+, EnumProperty(				Names.IndexStruct({Scopes.CoreUObject, Names.MakeTypename("FEnumProperty")}))
+, BoolProperty(				Names.IndexStruct({Scopes.CoreUObject, Names.MakeTypename("FBoolProperty")}))
+, StructProperty(			Names.IndexStruct({Scopes.CoreUObject, Names.MakeTypename("FStructProperty")}))
+, ClassProperty(			Names.IndexStruct({Scopes.CoreUObject, Names.MakeTypename("FClassProperty")}))
+, ObjectPropertyBase(		Names.IndexStruct({Scopes.CoreUObject, Names.MakeTypename("FObjectPropertyBase")}))
+, ArrayProperty(			Names.IndexStruct({Scopes.CoreUObject, Names.MakeTypename("FArrayProperty")}))
+, SetProperty(				Names.IndexStruct({Scopes.CoreUObject, Names.MakeTypename("FSetProperty")}))
+, MapProperty(				Names.IndexStruct({Scopes.CoreUObject, Names.MakeTypename("FMapProperty")}))
+, OptionalProperty(			Names.IndexStruct({Scopes.CoreUObject, Names.MakeTypename("FOptionalProperty")}))
+{}
+
 FCommonStructIds::FCommonStructIds(const FCommonScopeIds& Scopes, TIdIndexer<FSensitiveName>& Names)
 : Name(						Names.IndexStruct({Scopes.CoreUObject, Names.MakeTypename("Name")}))
 , Text(						Names.IndexStruct({Scopes.CoreUObject, Names.MakeTypename("Text")}))
 , Guid(						Names.IndexStruct({Scopes.CoreUObject, Names.MakeTypename("Guid")}))
 , FieldPath(				Names.IndexStruct({Scopes.CoreUObject, Names.MakeTypename("FieldPath")}))
+, Struct(					Names.IndexStruct({Scopes.CoreUObject, Names.MakeTypename("Struct")}))
+, Class(					Names.IndexStruct({Scopes.CoreUObject, Names.MakeTypename("Class")}))
+, Function(					Names.IndexStruct({Scopes.CoreUObject, Names.MakeTypename("Function")}))
+, ScriptStruct(				Names.IndexStruct({Scopes.CoreUObject, Names.MakeTypename("ScriptStruct")}))
 , SoftObjectPath(			Names.IndexStruct({Scopes.CoreUObject, Names.MakeTypename("SoftObjectPath")}))
 , ClassPtr(					Names.IndexStruct({Scopes.CoreUObject, Names.MakeTypename("ClassPtr")}))
 , ObjectPtr(				Names.IndexStruct({Scopes.CoreUObject, Names.MakeTypename("ObjectPtr")}))
@@ -350,6 +368,7 @@ FGlobals::FGlobals()
 , Scopes(Names)
 , EnumIds(Scopes, Names)
 , Structs(Scopes, Names)
+, Properties(Scopes, Names)
 , Typenames(Names)
 , Members(Names)
 , Numerals(Names)
@@ -402,19 +421,29 @@ static bool HasAny(uint64 Flags)
 
 ////////////////////////////////////////////////////////////////////////////////////////////////
 
-static FType IndexType(const UField* Field)
+static FScopeId IndexScope(const UField* Field)
 {
-	check(Field);
-	FTypenameId Name = GUE.Names.MakeTypename(Field->GetFName());
-	
 	// This wouldn't be needed in an intrusive or cached solution
 	TArray<FFlatScopeId, TInlineAllocator<64>> ReversedOuters;
 	for (const UObject* Outer = Field->GetOuter(); Outer; Outer = Outer->GetOuter())
 	{
 		ReversedOuters.Add(GUE.Names.NameScope(Outer->GetFName()));
 	}
-	
-	return { GUE.Names.NestReversedScopes(ReversedOuters), Name };
+	return GUE.Names.NestReversedScopes(ReversedOuters);
+}
+
+static FType IndexType(const UField* Field)
+{
+	return { IndexScope(Field), GUE.Names.MakeTypename(Field->GetFName()) };
+}
+
+static FType IndexType(const UField* Field, const char* Suffix)
+{
+	TUtf8StringBuilder<64> Str;
+	Field->GetFName().AppendString(Str);
+	Str.Append(Suffix);
+	FTypenameId Typename = GUE.Names.MakeTypename(*Str);
+	return { IndexScope(Field), Typename };
 }
 
 FType IndexStruct(const UStruct* Struct)
@@ -422,6 +451,11 @@ FType IndexStruct(const UStruct* Struct)
 	return IndexType(Struct);
 }
 
+FType IndexStructMeta(const UStruct* Struct)
+{
+	return IndexType(Struct, "_Meta");
+}
+
 static EMemberPresence GetOccupancy(const UStruct* Struct)
 {
 	if (Struct->HasAnyCastFlags(CASTCLASS_UScriptStruct))
@@ -4295,6 +4329,132 @@ FStructDeclarationPtr FPropertyBinder::Declare() const
 
 ////////////////////////////////////////////////////////////////////////////////////////////////
 
+// Todo: Merge with SavePropertyMeta below
+FDualStructId PropertyStructIdMeta(const FProperty* Property)
+{
+	FName MemberTypename = Property->GetFName();
+	uint64 Flags = Property->GetCastFlags();
+	if (HasAny<LeafMask>(Flags))
+	{
+		if (HasAny<CASTCLASS_FNumericProperty>(Flags))
+		{
+			return GUE.Properties.Property;
+		}
+		return HasAny<CASTCLASS_FEnumProperty>(Flags)
+			? GUE.Properties.EnumProperty
+			: GUE.Properties.BoolProperty;
+	}
+	else if (HasAny<CASTCLASS_FStructProperty>(Flags))
+	{
+		return GUE.Properties.StructProperty;
+	}
+	else if (HasAny<CASTCLASS_FClassProperty>(Flags))
+	{
+		return GUE.Properties.ClassProperty;
+	}
+	else if (HasAny<CASTCLASS_FObjectPropertyBase>(Flags))
+	{
+		return GUE.Properties.ObjectPropertyBase;
+	}
+	else if (HasAny<ContainerMask>(Flags))
+	{
+		if (HasAny<CASTCLASS_FArrayProperty>(Flags))
+		{
+			return GUE.Properties.ArrayProperty;
+		}
+		if (HasAny<CASTCLASS_FMapProperty>(Flags))
+		{
+			return GUE.Properties.MapProperty;
+		}
+		return HasAny<CASTCLASS_FSetProperty>(Flags)
+			? GUE.Properties.SetProperty
+			: GUE.Properties.OptionalProperty;
+	}
+	return GUE.Properties.Property;
+}
+
+// Meta version similar to FPropertyBinder::BindSingleProperty
+FBuiltStruct* SavePropertyMeta(const FField* Field, const FSaveContext& Ctx)
+{
+	FName MemberTypename = Field->GetFName();
+	uint64 Flags = Field->GetCastFlags();
+	if (HasAny<LeafMask>(Flags))
+	{
+		if (HasAny<CASTCLASS_FNumericProperty>(Flags))
+		{
+			return SaveStruct(Field, GUE.Properties.Property, Ctx);
+		}
+		return HasAny<CASTCLASS_FEnumProperty>(Flags)
+			? SaveStruct(Field, GUE.Properties.Property, Ctx) // Todo: FEnumPropertyBinding
+			: SaveStruct(Field, GUE.Properties.BoolProperty, Ctx);
+	}
+	else if (HasAny<CASTCLASS_FStructProperty>(Flags))
+	{
+		return SaveStruct(Field, GUE.Properties.StructProperty, Ctx);
+	}
+	else if (HasAny<CASTCLASS_FClassProperty>(Flags))
+	{
+		return SaveStruct(Field, GUE.Properties.ClassProperty, Ctx);
+	}
+	else if (HasAny<CASTCLASS_FObjectPropertyBase>(Flags))
+	{
+		return SaveStruct(Field, GUE.Properties.ObjectPropertyBase, Ctx);
+	}
+	else if (HasAny<ContainerMask>(Flags))
+	{
+		if (HasAny<CASTCLASS_FArrayProperty>(Flags))
+		{
+			return SaveStruct(Field, GUE.Properties.Property, Ctx); // Todo: FArrayPropertyBinding
+		}
+		if (HasAny<CASTCLASS_FMapProperty>(Flags))
+		{
+			return SaveStruct(Field, GUE.Properties.Property, Ctx); // Todo: FMapPropertyBinding
+		}
+		return HasAny<CASTCLASS_FSetProperty>(Flags)
+			? SaveStruct(Field, GUE.Properties.Property, Ctx) // Todo: FSetPropertyBinding
+			: SaveStruct(Field, GUE.Properties.Property, Ctx); // Todo: FOptionalBinding
+	}
+	else if (HasAny<StringMask>(Flags))
+	{
+		return SaveStruct(Field, GUE.Properties.Property, Ctx); // Todo: FStringPropertyBinding
+	}
+	else if (HasAny<MiscMask>(Flags))
+	{
+		return SaveStruct(Field, GUE.Properties.Property, Ctx); // Todo: FInterfaceBinding and FDelegateBindings
+	}
+#if WITH_VERSE_VM
+	else if (HasAny<CASTCLASS_FVValueProperty | CASTCLASS_FVRestValueProperty>(Flags))
+	{
+		// return HasAny<CASTCLASS_FVValueProperty>(Flags)
+		// 	? BindVValue(static_cast<FVValueProperty*>(Property))
+		// 	: BindVRestValue(static_cast<FVRestValueProperty*>(Property)); 
+	}
+#else // Verse BPVM
+	else
+	{
+		if (MemberTypename == "VerseFunctionProperty") // Todo: FVerseFunctionProperty
+		{
+			return SaveStruct(Field, GUE.Properties.Property, Ctx);
+		}
+		else if (MemberTypename == "VerseDynamicProperty") // Todo: FVerseDynamicProperty
+		{
+			return SaveStruct(Field, GUE.Properties.Property, Ctx);
+		}
+		else if (MemberTypename == "ReferenceProperty") // Todo: FReferenceProperty
+		{
+			return SaveStruct(Field, GUE.Properties.Property, Ctx);
+		}
+	}
+#endif
+
+	// Todo: check for complete coverage after implementing all missing FProperty bindings.
+	// checkf(false, TEXT("Unrecognized class cast flags %llx in %s %s"),
+	// Flags, *MemberTypename.ToString(), *Field->GetName());
+	return SaveStruct(Field, GUE.Properties.Property, Ctx);
+}
+
+////////////////////////////////////////////////////////////////////////////////////////////////
+
 static EPropertyFlags GDontBind;
 
 static bool ShouldBind(const FProperty* Property)
@@ -4325,6 +4485,18 @@ bool ShouldBind(const UStruct* Struct)
 	{
 		return !!(static_cast<const UScriptStruct*>(Struct)->StructFlags & STRUCT_SerializeNative);
 	}
+#if PP_OVERRIDE_OBJECT_LOADING
+	// Just as we need to include custom serialize structs without properties for BindMembers and GetSuperToBind,
+	// we also need to include all custom bound structs.
+	// temp: for now it is enough to check for the global common structs, but re-indexing every time is wasteful.
+	// todo: optimize and simplify and check for structs where the custom binding may not exist yet,
+	// probably needs a type trait...
+	FBindId Id = GUE.Names.IndexBindId(IndexType(Struct));
+	if (GUE.Customs.FindStruct(Id))
+	{
+		return true;
+	}
+#endif
 	return false;
 }
 
@@ -4427,6 +4599,14 @@ static TArray<FBindId> GetStructsToSkip(EBatchType BatchType)
 	}
 }
 
+static void InitPropertiesMeta()
+{
+	GUE.Defaults.BindZeroes(GUE.Properties.Property, sizeof(FProperty), alignof(FProperty));
+	GUE.Defaults.BindZeroes(GUE.Properties.BoolProperty, sizeof(FBoolProperty), alignof(FBoolProperty));
+	GUE.Defaults.BindZeroes(GUE.Properties.StructProperty, sizeof(FStructProperty), alignof(FStructProperty));
+	GUE.Defaults.BindZeroes(GUE.Properties.ClassProperty, sizeof(FClassProperty), alignof(FClassProperty));
+}
+
 static void InitBatchedProperties(EBatchType BatchType)
 {
 	GUE.Defaults.BindZeroes(GUE.Structs.FieldPath, sizeof(FFieldPath), alignof(FFieldPath));
@@ -4445,6 +4625,7 @@ static void InitBatchedProperties(EBatchType BatchType)
 
 void SchemaBindAllTypes(EBindMode Mode, EBatchType BatchType)
 {
+	InitPropertiesMeta();
 	InitBatchedProperties(BatchType);
 
 	GDontBind = GetPropertiesToSkip(Mode);
@@ -4484,6 +4665,10 @@ void SchemaBindAllTypes(EBindMode Mode, EBatchType BatchType)
 		}
 
 		BindStruct(Id, Struct);
+#if PP_OVERRIDE_OBJECT_LOADING
+		BindStructMeta(GUE.Names.IndexBindId(IndexStructMeta(Struct)), Struct);
+#endif
+
 		GScratchRanges.Reset();
 	}
 	GScratchRanges.Empty();
@@ -4506,6 +4691,22 @@ struct FPlainPropsBindingScope
 		, VerseFunction(GUE.Structs.VerseFunction)
 		, DynamicallyTypedValue(GUE.Structs.DynamicallyTypedValue)
 		, ReferencePropertyValue(GUE.Structs.ReferencePropertyValue)
+#if PP_OVERRIDE_OBJECT_LOADING
+		, ClassFlags()
+		, FunctionFlags()
+		, StructFlags()
+		, PropertyFlags()
+		, Property(GUE.Properties.Property)
+		, BoolProperty(GUE.Properties.BoolProperty)
+		, StructProperty(GUE.Properties.StructProperty)
+		, ClassProperty(GUE.Properties.ClassProperty)
+		, ObjectPropertyBase(GUE.Properties.ObjectPropertyBase)
+		, Struct(GUE.Structs.Struct)
+		, Class(GUE.Structs.Class)
+		, Function(GUE.Structs.Function)
+		, ScriptStruct(GUE.Structs.ScriptStruct)
+		, ImplementedInterface()
+#endif
 	{
 		GUE.Defaults.BindStatic(Transform.BindId, &DefaultTransform);
 		GUE.Defaults.BindZeroes<FGuid>(Guid.BindId);
@@ -4533,6 +4734,23 @@ struct FPlainPropsBindingScope
 	TScopedStructBinding<FVerseFunction, FDefaultRuntime> VerseFunction;
 	TScopedStructBinding<::UE::FDynamicallyTypedValue, FDefaultRuntime> DynamicallyTypedValue;
 	TScopedStructBinding<FReferencePropertyValue, FDefaultRuntime> ReferencePropertyValue;
+
+#if PP_OVERRIDE_OBJECT_LOADING
+	TScopedEnumDeclaration<EClassFlags, EEnumMode::Flag, FDefaultRuntime> ClassFlags;
+	TScopedEnumDeclaration<EFunctionFlags, EEnumMode::Flag, FDefaultRuntime> FunctionFlags;
+	TScopedEnumDeclaration<EStructFlags, EEnumMode::Flag, FDefaultRuntime> StructFlags;
+	TScopedEnumDeclaration<EPropertyFlags, EEnumMode::Flag, FDefaultRuntime> PropertyFlags;
+	TScopedStructBinding<FProperty, FDefaultRuntime> Property;
+	TScopedStructBinding<FBoolProperty, FDefaultRuntime> BoolProperty;
+	TScopedStructBinding<FStructProperty, FDefaultRuntime> StructProperty;
+	TScopedStructBinding<FClassProperty, FDefaultRuntime> ClassProperty;
+	TScopedStructBinding<FObjectPropertyBase, FDefaultRuntime> ObjectPropertyBase;
+	TScopedStructBinding<UStruct, FDefaultRuntime> Struct;
+	TScopedStructBinding<UClass, FDefaultRuntime> Class;
+	TScopedStructBinding<UFunction, FDefaultRuntime> Function;
+	TScopedStructBinding<UScriptStruct, FDefaultRuntime> ScriptStruct;
+	TScopedStructBinding<FImplementedInterface, FDefaultRuntime> ImplementedInterface;
+#endif
 };
 
 } // namespace PlainProps::UE
diff --git a/Engine/Plugins/Experimental/PlainPropsUObject/Source/Private/PlainPropsUObjectBindingsInternal.h b/Engine/Plugins/Experimental/PlainPropsUObject/Source/Private/PlainPropsUObjectBindingsInternal.h
--- a/Engine/Plugins/Experimental/PlainPropsUObject/Source/Private/PlainPropsUObjectBindingsInternal.h
+++ b/Engine/Plugins/Experimental/PlainPropsUObject/Source/Private/PlainPropsUObjectBindingsInternal.h
@@ -3,19 +3,28 @@
 #pragma once
 
 #include "PlainPropsTypes.h"
+#include "PlainPropsBind.h"
 #include "Logging/LogMacros.h"
 
 DECLARE_LOG_CATEGORY_EXTERN(LogPlainPropsUObject, Log, All);
 
+class FField;
+class FProperty;
 class UStruct;
+namespace PlainProps { struct FSaveContext; };
 
 namespace PlainProps::UE
 {
 
 FType			IndexStruct(const UStruct* Struct);
-void			BindStruct(FBindId Id, const UStruct* Struct); // to bind dynamic structs as they are being loaded
 bool			ShouldBind(const UStruct* Struct);
 const UStruct*	SkipEmptyBases(const UStruct* Struct);
 
+// Helpers for native UField and FProperty meta bindings
+FType			IndexStructMeta(const UStruct* Struct);
+void 			BindStructMeta(FBindId Id, const UStruct* Struct);
+FDualStructId 	PropertyStructIdMeta(const FProperty* Property);
+FBuiltStruct* 	SavePropertyMeta(const FField* Field, const FSaveContext& Ctx);
+
 } // namespace PlainProps::UE
 
diff --git a/Engine/Plugins/Experimental/PlainPropsUObject/Source/Private/PlainPropsUObjectMetaBindings.cpp b/Engine/Plugins/Experimental/PlainPropsUObject/Source/Private/PlainPropsUObjectMetaBindings.cpp
new file mode 100644
--- /dev/null
+++ b/Engine/Plugins/Experimental/PlainPropsUObject/Source/Private/PlainPropsUObjectMetaBindings.cpp
@@ -0,0 +1,703 @@
+// Copyright Epic Games, Inc. All Rights Reserved.
+
+#include "PlainPropsUObjectMetaBindingsInternal.h"
+#include "PlainPropsInternalPrivateMemberPtr.h"
+#include "PlainPropsUObjectBindingsInternal.h"
+#include "Misc/DefinePrivateMemberPtr.h"
+#include "UObject/Class.h"
+#include "UObject/CoreNetTypes.h"
+#include "UObject/LinkerLoad.h"
+#include "UObject/ObjectMacros.h"
+#include "UObject/ObjectPtr.h"
+#include "UObject/Package.h"
+
+// Temp hacks. Long-term either add FProperty getters for ctor/dtor/hash function pointers 
+// and delegate APIs for non-intrusive serialization or integrate PlainProps into Core/CoreUObject
+UE_DEFINE_PRIVATE_MEMBER_PTR(TObjectPtr<UStruct>, GSuperStruct, UStruct, SuperStruct);
+
+PRAGMA_DISABLE_DEPRECATION_WARNINGS
+PP_DEFINE_PRIVATE_MEMBER_NAMESPACE_BEGIN(UClass_Private)
+PP_DEFINE_PRIVATE_MEMBER_PTR(UClass, ClassDefaultObject, TObjectPtr<UObject>);
+#if USTRUCT_FAST_ISCHILDOF_IMPL == USTRUCT_ISCHILDOF_STRUCTARRAY
+PP_DEFINE_PRIVATE_MEMBER_PTR(UStruct, ReinitializeBaseChainArray, void());
+#endif
+PP_DEFINE_PRIVATE_MEMBER_NAMESPACE_END()
+
+PP_DEFINE_PRIVATE_MEMBER_NAMESPACE_BEGIN(FProperty_Private)
+PP_DEFINE_PRIVATE_MEMBER_PTR(FProperty, ElementSize,					int32);
+PP_DEFINE_PRIVATE_MEMBER_PTR(FProperty, BlueprintReplicationCondition,	TEnumAsByte<ELifetimeCondition>);
+PP_DEFINE_PRIVATE_MEMBER_PTR(FProperty, Offset_Internal,				int32);
+PP_DEFINE_PRIVATE_MEMBER_NAMESPACE_END()
+
+PP_DEFINE_PRIVATE_MEMBER_NAMESPACE_BEGIN(FBoolProperty_Private)
+PP_DEFINE_PRIVATE_MEMBER_PTR(FBoolProperty, FieldSize,	uint8);
+PP_DEFINE_PRIVATE_MEMBER_PTR(FBoolProperty, ByteOffset,	uint8);
+PP_DEFINE_PRIVATE_MEMBER_PTR(FBoolProperty, ByteMask,	uint8);
+PP_DEFINE_PRIVATE_MEMBER_PTR(FBoolProperty, FieldMask,	uint8);
+PP_DEFINE_PRIVATE_MEMBER_NAMESPACE_END()
+PRAGMA_ENABLE_DEPRECATION_WARNINGS
+
+namespace PlainProps::UE
+{
+
+////////////////////////////////////////////////////////////////////////////////////////////////
+// Native FProperty meta bindings
+
+FPropertyBinding::FPropertyBinding(TCustomSpecifier<Ids, NumMembers>& Spec)
+	: PropertyFlagsId(GetEnumId<Ids, EPropertyFlags>())
+	, LifetimeConditionId(Ids::IndexEnum(GUE.Scopes.CoreUObject, "ELifetimeCondition"))
+{
+	// All FFields are FProperties so we can fold all FField members here for simplicity and efficiency
+	Spec.Members[M::NamePrivate] = FDeclId(GUE.Structs.Name);
+	Spec.Members[M::ArrayDim] = Specify<int32>();
+	Spec.Members[M::ElementSize] = Specify<int32>();
+	Spec.Members[M::PropertyFlags] = Specify<EPropertyFlags>(PropertyFlagsId);
+	Spec.Members[M::RepIndex] = Specify<uint16>();
+	Spec.Members[M::BlueprintReplicationCondition] = Specify<ELifetimeCondition>(LifetimeConditionId);
+	// Todo: Save and load editor only FField::MetaDataMap
+	// TMap<FName, FString>* MetaDataMap;
+}
+
+void FPropertyBinding::Save(FMemberBuilder& Dst, const FProperty& Src, const FProperty*, const FSaveContext& Ctx) const
+{
+	using namespace FProperty_Private;
+	Dst.AddStruct(MemberIds[M::NamePrivate], GUE.Structs.Name, SaveStruct(&Src.NamePrivate, GUE.Structs.Name, Ctx));
+	Dst.Add(MemberIds[M::ArrayDim], Src.ArrayDim);
+	Dst.Add(MemberIds[M::ElementSize], Src.*_ElementSize);
+	EPropertyFlags SaveFlags = Src.PropertyFlags & ~CPF_ComputedFlags;
+	Dst.AddEnum(MemberIds[M::PropertyFlags], PropertyFlagsId, SaveFlags);
+	Dst.Add(MemberIds[M::RepIndex], Src.RepIndex);
+	ELifetimeCondition SaveBlueprintReplicationCondition = (Src.*_BlueprintReplicationCondition).GetValue();
+	Dst.AddEnum(MemberIds[M::BlueprintReplicationCondition], LifetimeConditionId, SaveBlueprintReplicationCondition);
+	// TMap<FName, FString>* MetaDataMap;
+}
+
+void FPropertyBinding::Load(FProperty& Dst, FStructLoadView Src, ECustomLoadMethod Method) const
+{
+	using namespace FProperty_Private;
+	FMemberLoader Members(Src);
+	LoadStruct(&Dst.NamePrivate, Members.GrabStruct());
+	Dst.ArrayDim = Members.GrabLeaf().As<int32>();
+	Dst.*_ElementSize = Members.GrabLeaf().As<int32>();
+	EPropertyFlags LoadFlags = Members.GrabLeaf().As<EPropertyFlags>();
+	Dst.PropertyFlags = (LoadFlags & ~CPF_ComputedFlags) | (Dst.PropertyFlags & CPF_ComputedFlags);
+	Dst.RepIndex = Members.GrabLeaf().As<uint16>();
+	ELifetimeCondition LoadBlueprintReplicationCondition = Members.GrabLeaf().As<ELifetimeCondition>();
+	Dst.*_BlueprintReplicationCondition = LoadBlueprintReplicationCondition;
+	// TMap<FName, FString>* MetaDataMap;
+	Dst.*_Offset_Internal = 0;
+	Dst.DestructorLinkNext = nullptr;
+	check(!Members.HasMore());
+}
+
+FBoolPropertyBinding::FBoolPropertyBinding(TCustomSpecifier<Ids, NumMembers>& Spec)
+{
+	Spec.Super = FDeclId(GUE.Properties.Property);
+	Spec.Members[M::FieldSize] = Specify<uint8>();
+	Spec.Members[M::ByteOffset] = Specify<uint8>();
+	Spec.Members[M::ByteMask] = Specify<uint8>();
+	Spec.Members[M::FieldMask] = Specify<uint8>();
+}
+
+void FBoolPropertyBinding::Save(FMemberBuilder& Dst, const FBoolProperty& Src, const FBoolProperty*, const FSaveContext& Ctx) const
+{
+	using namespace FBoolProperty_Private;
+	SaveSuper(Dst, &Src, FBindId(GUE.Properties.Property), Ctx);
+	Dst.Add(MemberIds[M::FieldSize], Src.*_FieldSize);
+	Dst.Add(MemberIds[M::ByteOffset], Src.*_ByteOffset);
+	Dst.Add(MemberIds[M::ByteMask], Src.*_ByteMask);
+	Dst.Add(MemberIds[M::FieldMask], Src.*_FieldMask);
+}
+
+void FBoolPropertyBinding::Load(FBoolProperty& Dst, FStructLoadView Src, ECustomLoadMethod Method) const
+{
+	using namespace FBoolProperty_Private;
+	FMemberLoader Members(Src);
+	TryLoadSuper(&Dst, Members);
+	Dst.*_FieldSize = Members.GrabLeaf().As<uint8>();
+	Dst.*_ByteOffset = Members.GrabLeaf().As<uint8>();
+	Dst.*_ByteMask = Members.GrabLeaf().As<uint8>();
+	Dst.*_FieldMask = Members.GrabLeaf().As<uint8>();
+	check(!Members.HasMore());
+}
+
+FStructPropertyBinding::FStructPropertyBinding(TCustomSpecifier<Ids, NumMembers>& Spec)
+{
+	Spec.Super = FDeclId(GUE.Properties.Property);
+	Spec.Members[M::Struct] = FDeclId(GUE.Structs.ObjectPtr);
+}
+
+void FStructPropertyBinding::Save(FMemberBuilder& Dst, const FStructProperty& Src, const FStructProperty*, const FSaveContext& Ctx) const
+{
+	SaveSuper(Dst, &Src, FBindId(GUE.Properties.Property), Ctx);
+	Dst.AddStruct(MemberIds[M::Struct], GUE.Structs.ObjectPtr, SaveStruct(&Src.Struct, GUE.Structs.ObjectPtr, Ctx));
+}
+
+void FStructPropertyBinding::Load(FStructProperty& Dst, FStructLoadView Src, ECustomLoadMethod Method) const
+{
+	FMemberLoader Members(Src);
+	TryLoadSuper(&Dst, Members);
+	LoadStruct(&Dst.Struct, Members.GrabStruct());
+	check(!Members.HasMore());
+}
+
+FClassPropertyBinding::FClassPropertyBinding(TCustomSpecifier<Ids, NumMembers>& Spec)
+{
+	Spec.Super = FDeclId(GUE.Properties.Property);
+	Spec.Members[M::MetaClass] = FDeclId(GUE.Structs.ObjectPtr);
+}
+
+void FClassPropertyBinding::Save(FMemberBuilder& Dst, const FClassProperty& Src, const FClassProperty*, const FSaveContext& Ctx) const
+{
+	SaveSuper(Dst, &Src, FBindId(GUE.Properties.Property), Ctx);
+	Dst.AddStruct(MemberIds[M::MetaClass], GUE.Structs.ObjectPtr, SaveStruct(&Src.MetaClass, GUE.Structs.ObjectPtr, Ctx));
+}
+
+void FClassPropertyBinding::Load(FClassProperty& Dst, FStructLoadView Src, ECustomLoadMethod Method) const
+{
+	FMemberLoader Members(Src);
+	TryLoadSuper(&Dst, Members);
+	LoadStruct(&Dst.MetaClass, Members.GrabStruct());
+	// Todo: AddReferencingProperty requires private header "UObject/LinkerPlaceholderClass.h"
+	// #if USE_CIRCULAR_DEPENDENCY_LOAD_DEFERRING
+	// 	if (ULinkerPlaceholderClass* PlaceholderClass = Cast<ULinkerPlaceholderClass>(Dst.MetaClass))
+	// 	{
+	// 		PlaceholderClass->AddReferencingProperty(&Dst);
+	// 	}
+	// #endif
+	check(!Members.HasMore());
+}
+
+FObjectPropertyBaseBinding::FObjectPropertyBaseBinding(TCustomSpecifier<Ids, NumMembers>& Spec)
+{
+	Spec.Super = FDeclId(GUE.Properties.Property);
+	Spec.Members[M::PropertyClass] = FDeclId(GUE.Structs.ObjectPtr);
+}
+
+void FObjectPropertyBaseBinding::Save(FMemberBuilder& Dst, const FObjectPropertyBase& Src, const FObjectPropertyBase*, const FSaveContext& Ctx) const
+{
+	SaveSuper(Dst, &Src, FBindId(GUE.Properties.Property), Ctx);
+	Dst.AddStruct(MemberIds[M::PropertyClass], GUE.Structs.ObjectPtr, SaveStruct(&Src.PropertyClass, GUE.Structs.ObjectPtr, Ctx));
+}
+
+void FObjectPropertyBaseBinding::Load(FObjectPropertyBase& Dst, FStructLoadView Src, ECustomLoadMethod Method) const
+{
+	FMemberLoader Members(Src);
+	TryLoadSuper(&Dst, Members);
+	LoadStruct(&Dst.PropertyClass, Members.GrabStruct());
+	check(!Members.HasMore());
+}
+
+////////////////////////////////////////////////////////////////////////////////////////////////
+
+UStructBinding::UStructBinding(TCustomSpecifier<Ids, NumMembers>& Spec)
+{
+	Spec.Members[M::SuperStruct] = FDeclId(GUE.Structs.ObjectPtr);
+	// Children is an array of UFields, these are always UFunctions
+	Spec.Members[M::Children] = DefaultRangeOf(FDeclId(GUE.Structs.ObjectPtr));
+	// This UStruct base binding saves and loads the linked list of FProperties as an array.
+	// The derived meta classes saves and loads the actual FProperty data as named schema members.
+	Spec.Members[M::ChildPropertyClasses] = DefaultRangeOf(FDeclId(GUE.Structs.Name));
+}
+
+FTypedRange UStructBinding::SaveChildrenStatic(const UField* Children, int32 Num, const FSaveContext& Ctx)
+{
+	FBindId Id = GUE.Structs.ObjectPtr;
+	FStructRangeSaver Out(Ctx.Scratch, Num);
+	for (const UField* Child = Children; Child; Child = Child->Next)
+	{
+		Out.AddItem(SaveStruct(&Child, Id, Ctx));
+	}
+	return Out.Finalize(MakeStructRangeSchema(DefaultRangeMax, Id));
+}
+
+FTypedRange UStructBinding::SaveChildPropertyClassesStatic(const FField* ChildProperties, int32 Num, const FSaveContext& Ctx)
+{
+	FBindId Id = GUE.Structs.Name;
+	FStructRangeSaver Out(Ctx.Scratch, Num);
+	for (const FField* Child = ChildProperties; Child; Child = Child->Next)
+	{
+		FName ClassName = Child->GetClass()->GetFName();
+		Out.AddItem(SaveStruct(&ClassName, Id, Ctx));
+	}
+	return Out.Finalize(MakeStructRangeSchema(DefaultRangeMax, Id));
+}
+
+void UStructBinding::Save(FMemberBuilder& Dst, const UStruct& Src, const UStruct*, const FSaveContext& Ctx) const
+{
+	Dst.AddStruct(MemberIds[M::SuperStruct], GUE.Structs.ObjectPtr, SaveStruct(&(Src.*GSuperStruct), GUE.Structs.ObjectPtr, Ctx));
+	if (Src.Children)
+	{
+		int32 Num = 0;
+		for (const UField* Child = Src.Children; Child; Child = Child->Next, ++Num);
+
+		Dst.AddRange(MemberIds[M::Children], SaveChildrenStatic(Src.Children, Num, Ctx));
+	}
+	if (Src.ChildProperties)
+	{
+		int32 Num = 0;
+		for (const FField* Child = Src.ChildProperties; Child; Child = Child->Next, ++Num);
+
+		Dst.AddRange(MemberIds[M::ChildPropertyClasses], SaveChildPropertyClassesStatic(Src.ChildProperties, Num, Ctx));
+	}
+}
+
+void UStructBinding::LoadChildrenStatic(TObjectPtr<UField>& Children, FStructRangeLoadView Src)
+{
+	Children = nullptr;
+	UField* Previous = nullptr;
+	for (FStructLoadView SrcIt : Src)
+	{
+		TObjectPtr<UField> Current;
+		LoadStruct(&Current, SrcIt);
+		if (Current)
+		{
+			if (!Children)
+			{
+				Children = Current;
+			}
+			if (Previous)
+			{
+				Previous->Next = Current;
+			}
+			Previous = Current;
+		}
+	}
+	if (Previous)
+	{
+		Previous->Next = nullptr;
+	}
+}
+
+void UStructBinding::LoadChildPropertyClassesStatic(UStruct& Dst, FStructRangeLoadView Src)
+{
+	Dst.ChildProperties = nullptr;
+	FField* Previous = nullptr;
+	for (FStructLoadView SrcIt : Src)
+	{
+		FName PropertyTypeName;
+		LoadStruct(&PropertyTypeName, SrcIt);
+		FField* Current = FField::Construct(PropertyTypeName, &Dst, NAME_None);
+		check(Current);
+		if (!Dst.ChildProperties)
+		{
+			Dst.ChildProperties = Current;
+		}
+		if (Previous)
+		{
+			Previous->Next = Current;
+		}
+		Previous = Current;
+	}
+	if (Previous)
+	{
+		Previous->Next = nullptr;
+	}
+}
+
+void UStructBinding::Load(UStruct& Dst, FStructLoadView Src, ECustomLoadMethod Method) const
+{
+	using namespace UClass_Private;
+
+	FMemberLoader Members(Src);
+	LoadStruct(&(Dst.*GSuperStruct), Members.GrabStruct());
+#if USTRUCT_FAST_ISCHILDOF_IMPL == USTRUCT_ISCHILDOF_STRUCTARRAY
+	(Dst.*_ReinitializeBaseChainArray)();
+#endif
+	if (Members.HasMore())
+	{
+		check(Members.PeekKind() == EMemberKind::Range);
+		if (Members.PeekName() == MemberIds[M::Children])
+		{
+			LoadChildrenStatic(Dst.Children, Members.GrabRange().AsStructs());
+		}
+	}
+	if (Members.HasMore())
+	{
+		check(Members.PeekKind() == EMemberKind::Range);
+		check(Members.PeekName() == MemberIds[M::ChildPropertyClasses])
+		LoadChildPropertyClassesStatic(Dst, Members.GrabRange().AsStructs());
+	}
+	// UStruct::Serialize first serializes all FProperties, then calls link for all types but UClasses.
+	// This base custom binding does not load the FProperties (only their types for constructing them),
+	// so it always delegates the responsibility to call Link to the derived meta classes.
+	check(!Members.HasMore());
+}
+
+UClassBinding::UClassBinding(TCustomSpecifier<Ids, NumMembers>& Spec)
+: ClassFlagsId(GetEnumId<Ids, EClassFlags>())
+, ImplementedInterfaceId(GetStructDualId<Ids, FImplementedInterface>())
+{
+	Spec.Super = FDeclId(GUE.Structs.Struct);
+	Spec.Members[M::ClassFlags] = Specify<EClassFlags>(ClassFlagsId);
+	Spec.Members[M::ClassWithin] = FDeclId(GUE.Structs.ObjectPtr);
+	Spec.Members[M::ClassConfigName] = FDeclId(GUE.Structs.Name);
+#if WITH_EDITORONLY_DATA
+	Spec.Members[M::ClassGeneratedBy] = FDeclId(GUE.Structs.ObjectPtr);
+#endif
+	Spec.Members[M::Interfaces] = DefaultRangeOf(FDeclId(ImplementedInterfaceId));
+}
+
+FTypedRange UClassBinding::SaveInterfacesStatic(FBindId Id, TConstArrayView<FImplementedInterface> Interfaces, const FSaveContext& Ctx)
+{
+	FStructRangeSaver Out(Ctx.Scratch, Interfaces.Num());
+	for (const FImplementedInterface& Interface : Interfaces)
+	{
+		Out.AddItem(SaveStruct(&Interface, Id, Ctx));
+	}
+	return Out.Finalize(MakeStructRangeSchema(DefaultRangeMax, Id));
+}
+
+void UClassBinding::LoadInterfacesStatic(TArray<FImplementedInterface>& Interfaces, FStructRangeLoadView Src)
+{
+	for (FStructLoadView SrcIt : Src)
+	{
+		unimplemented();
+	}
+}
+
+void UClassBinding::Save(FMemberBuilder& Dst, const UClass& Src, const UClass*, const FSaveContext& Ctx) const
+{
+	SaveSuper(Dst, &Src, FBindId(GUE.Structs.Struct), Ctx);
+	EClassFlags SavedClassFlags = Src.ClassFlags & ~CLASS_ShouldNeverBeLoaded;
+	Dst.AddEnum(MemberIds[M::ClassFlags], ClassFlagsId, SavedClassFlags);
+	Dst.AddStruct(MemberIds[M::ClassWithin], GUE.Structs.ObjectPtr, SaveStruct(&Src.ClassWithin, GUE.Structs.ObjectPtr, Ctx));
+	Dst.AddStruct(MemberIds[M::ClassConfigName], GUE.Structs.Name, SaveStruct(&Src.ClassConfigName, GUE.Structs.Name, Ctx));
+#if WITH_EDITORONLY_DATA
+	Dst.AddStruct(MemberIds[M::ClassGeneratedBy], GUE.Structs.ObjectPtr, SaveStruct(&Src.ClassGeneratedBy, GUE.Structs.ObjectPtr, Ctx));
+#endif
+	if (Src.Interfaces.Num())
+	{
+		Dst.AddRange(MemberIds[M::Interfaces], SaveInterfacesStatic(FBindId(ImplementedInterfaceId), Src.Interfaces, Ctx));
+	}
+}
+
+void UClassBinding::Load(UClass& Dst, FStructLoadView Src, ECustomLoadMethod Method) const
+{
+	FMemberLoader Members(Src);
+
+#if PP_OVERRIDE_OBJECT_LOADING
+	UnhashObject(&Dst);
+	TryLoadSuper(&Dst, Members);
+	HashObject(&Dst);
+#endif
+
+	// UClass::Serialize calls Super::Serialize here to serialize all FProperties without yet calling Link
+
+	EClassFlags LoadedClassFlags = Members.GrabLeaf().As<EClassFlags>();
+	Dst.ClassFlags = LoadedClassFlags & ~CLASS_ShouldNeverBeLoaded;
+	LoadStruct(&Dst.ClassWithin, Members.GrabStruct());
+	LoadStruct(&Dst.ClassConfigName, Members.GrabStruct());
+#if WITH_EDITORONLY_DATA
+	LoadStruct(&Dst.ClassGeneratedBy, Members.GrabStruct());
+#endif
+
+	// Generate the FuncMap on load, it is only depending on the already loaded UFunction children,
+	// and it is not dependent on the not-yet-loaded FProperties.
+	// In the legacy system the FuncMap is explicitly serialized before the class flags.
+	for (UField* Child = Dst.Children; Child; Child = Child->Next)
+	{
+		UFunction* Function = CastChecked<UFunction>(Child);
+		Dst.AddFunctionToFunctionMap(Function, Function->GetFName());
+	}
+
+	// UClass::Serialize calls Link here AFTER serializing all FProperties and class flags,
+	// but BEFORE serializing the Interfaces and ClassDefaultObject.
+	// This base custom binding does not load the FProperties (only their types for constructing them),
+	// so it always delegates the responsibility to call Link to the derived meta classes.
+
+	// Don't save and load ClassDefaultObject from this base custom binding,
+	// it must be loaded explicitly after Link has been called, which happens right after
+	// loading FProperties in the meta binding that derives from UClass.
+
+	// Todo: Since VER_UE4_UCLASS_SERIALIZE_INTERFACES_AFTER_LINKING UClass::Serialize very explicitly 
+	// serializes the interfaces AFTER calling Link.
+	// Loading them here in the base custom binding will take place before the derived meta binding calls Link!!!
+	if (Members.HasMore())
+	{
+		check(Members.PeekKind() == EMemberKind::Range);
+		if (Members.PeekName() == MemberIds[M::Interfaces])
+		{
+			LoadInterfacesStatic(Dst.Interfaces, Members.GrabRange().AsStructs());
+		}
+	}
+	
+	// Todo: Since VER_UE4_ADD_COOKED_TO_UCLASS the legacy system serializes a cooked bool,
+	// this should already be known on the package level.
+	Dst.bCooked = false;
+}
+
+UFunctionBinding::UFunctionBinding(TCustomSpecifier<Ids, NumMembers>& Spec)
+: FunctionFlagsId(GetEnumId<Ids, EFunctionFlags>())
+{
+	Spec.Super = FDeclId(GUE.Structs.Struct);
+	Spec.Members[M::FunctionFlags] = Specify<EFunctionFlags>(FunctionFlagsId);
+	Spec.Members[M::EventGraphFunction] = FDeclId(GUE.Structs.ObjectPtr);
+	Spec.Members[M::EventGraphCallOffset] = Specify<int32>();
+}
+
+void UFunctionBinding::Save(FMemberBuilder& Dst, const UFunction& Src, const UFunction*, const FSaveContext& Ctx) const
+{
+	SaveSuper(Dst, &Src, FBindId(GUE.Structs.Struct), Ctx);
+	Dst.AddEnum(MemberIds[M::FunctionFlags], FunctionFlagsId, Src.FunctionFlags);
+	Dst.AddStruct(MemberIds[M::EventGraphFunction], GUE.Structs.ObjectPtr, SaveStruct(&Src.EventGraphFunction, GUE.Structs.ObjectPtr, Ctx));
+	Dst.Add(MemberIds[M::EventGraphCallOffset], Src.EventGraphCallOffset);
+}
+
+void UFunctionBinding::Load(UFunction& Dst, FStructLoadView Src, ECustomLoadMethod Method) const
+{
+	FMemberLoader Members(Src);
+	TryLoadSuper(&Dst, Members);
+	Dst.FunctionFlags = Members.GrabLeaf().As<EFunctionFlags>();
+	LoadStruct(&Dst.EventGraphFunction, Members.GrabStruct());
+	Dst.EventGraphCallOffset = Members.GrabLeaf().As<int32>();
+	// Link is called from the derived meta class.
+	check(!Members.HasMore());
+}
+
+UScriptStructBinding::UScriptStructBinding(TCustomSpecifier<Ids, NumMembers>& Spec)
+: StructFlagsId(GetEnumId<Ids, EStructFlags>())
+{
+	Spec.Super = FDeclId(GUE.Structs.Struct);
+	Spec.Members[M::StructFlags] = Specify<EStructFlags>(StructFlagsId);
+}
+
+void UScriptStructBinding::Save(FMemberBuilder& Dst, const UScriptStruct& Src, const UScriptStruct*, const FSaveContext& Ctx) const
+{
+	SaveSuper(Dst, &Src, FBindId(GUE.Structs.Struct), Ctx);
+	EStructFlags SavedStructFlags = EStructFlags(Src.StructFlags & ~STRUCT_ComputedFlags);
+	Dst.AddEnum(MemberIds[M::StructFlags], StructFlagsId, SavedStructFlags);
+}
+
+void UScriptStructBinding::Load(UScriptStruct& Dst, FStructLoadView Src, ECustomLoadMethod Method) const
+{
+	FMemberLoader Members(Src);
+	TryLoadSuper(&Dst, Members);
+	EStructFlags LoadedStructFlags = Members.GrabLeaf().As<EStructFlags>();
+	Dst.StructFlags = EStructFlags((LoadedStructFlags & ~STRUCT_ComputedFlags) | (Dst.StructFlags & STRUCT_ComputedFlags));
+	// Link is called from the derived meta class.
+	check(!Members.HasMore());
+}
+
+////////////////////////////////////////////////////////////////////////////////////////////////
+
+static bool IsNativeOrTransientStruct(const UStruct* Struct)
+{
+	check(!Struct->HasAnyFlags(RF_MarkAsNative));
+	if (Struct->HasAnyCastFlags(CASTCLASS_UClass))
+	{
+		EClassFlags Flags = static_cast<const UClass*>(Struct)->ClassFlags;
+		if (Flags & (CLASS_Native))
+		{
+			return true;
+		}
+	}
+	else if (Struct->HasAnyCastFlags(CASTCLASS_UScriptStruct))
+	{
+		EStructFlags Flags = static_cast<const UScriptStruct*>(Struct)->StructFlags;
+		if (Flags & (STRUCT_NoExport | STRUCT_Native))
+		{
+			return true;
+		}
+	}
+	else if (Struct->HasAnyCastFlags(CASTCLASS_UFunction))
+	{
+		EFunctionFlags Flags = static_cast<const UFunction*>(Struct)->FunctionFlags;
+		if (Flags & (FUNC_Native))
+		{
+			return true;
+		}
+	}
+	const UPackage* Package = Struct->GetPackage();
+	if (Package->HasAnyPackageFlags(PKG_CompiledIn))
+	{
+		return true;
+	}
+	if (Package == GetTransientPackage())
+	{
+		return true;
+	}
+	return false;
+}
+
+struct FStructBindingMeta final : ICustomBinding
+{
+	explicit FStructBindingMeta(FBindId InSuper, TConstArrayView<FMemberId> InMemberIds, TConstArrayView<FBindId> InBindIds)
+	: Super(InSuper)
+	, MemberIds(InMemberIds)
+	, BindIds(InBindIds)
+	{
+	}
+
+	virtual void SaveCustom(FMemberBuilder& Dst, const void* Src, FBaseline Base, const FSaveContext& Ctx) override
+	{
+		const void* Default = Base.Get();
+		if (!Default || DiffCustom(Src, Default, Ctx))
+		{
+			Save(Dst, *static_cast<const UStruct*>(Src), static_cast<const UStruct*>(Default), Ctx);
+		}
+	}
+
+	virtual void LoadCustom(void* Dst, FStructLoadView Src, ECustomLoadMethod Method) const override
+	{
+		check(Method == ECustomLoadMethod::Assign);
+		Load(*static_cast<UStruct*>(Dst), Src);
+	}
+
+	virtual bool DiffCustom(const void* A, const void* B, const FBindContext&) const override
+	{
+		return A != B;
+	}
+
+	void Save(FMemberBuilder& Dst, const UStruct& Src, const UStruct*, const FSaveContext& Ctx) const
+	{
+		SaveSuper(Dst, &Src, Super, Ctx);
+		using namespace UClass_Private;
+		if (Src.ChildProperties)
+		{
+			const FMemberId* MemberId = &MemberIds[0];
+			const FBindId* BindId = &BindIds[0];
+			for (const FField* Child = Src.ChildProperties; Child; Child = Child->Next)
+			{
+				Dst.AddStruct(*MemberId++, *BindId++, SavePropertyMeta(Child, Ctx));
+			}
+		}
+		if (const UClass* Class = Cast<UClass>(&Src))
+		{
+			// Todo: non-composable hack, save CDO for UClass after the named and type declared meta properties
+			Dst.AddStruct(MemberIds.Last(), GUE.Structs.ObjectPtr, SaveStruct(&((*Class).*_ClassDefaultObject), GUE.Structs.ObjectPtr, Ctx));
+		}
+	}
+
+	void Load(UStruct& Dst, FStructLoadView Src) const
+	{
+		using namespace UClass_Private;
+		FMemberLoader Members(Src);
+		// First load the custom base binding members which will construct the FProperties,
+		// then load the named and typed meta properties into them
+		TryLoadSuper(&Dst, Members);
+
+		UClass* Class = Cast<UClass>(&Dst);
+		const int32 NumChildren = MemberIds.Num() - !!Class;
+
+		FField* Child = Dst.ChildProperties;
+		for (int32 I = 0; I < NumChildren; ++I)
+		{
+			check(Child);
+			check(Members.HasMore());
+			LoadStruct(Child, Members.GrabStruct());
+			Child = Child->Next;
+		}
+		// Explicitly call Link for all types immediately after loading all named and typed meta properties
+		Dst.Link(*Dst.GetLinker(), true);
+		if (Class)
+		{
+			// Todo: non-composable hack, load the CDO for UClass after the named and type declared meta properties
+			// Todo: The UClass interface arrray will also need to be loaded at this point in time
+			check(Members.HasMore());
+			LoadStruct(&((*Class).*_ClassDefaultObject), Members.GrabStruct());
+			FDualStructId Id = GUE.Structs.ObjectPtr;
+		}
+		check(!Members.HasMore());
+	}
+
+	FBindId Super;
+	TArray<FMemberId> MemberIds;
+	TArray<FBindId> BindIds;
+};
+
+class FPropertyBinderMeta
+{
+public:
+	FPropertyBinderMeta(FBindId Id)
+	: Owner(Id)
+	{
+	}
+
+	void DeclareAndBind(const UStruct* Struct)
+	{
+		FDualStructId Super = GetSuper(Struct);
+		BindAllMembers(Struct);
+		FStructDeclarationPtr Declaration = Declare(FDeclId(Super));
+		// Todo: Ownership / memory leak
+		ICustomBinding* Leak = new FStructBindingMeta(FBindId(Super), Names, Members);
+		GUE.Customs.BindStruct(Owner, *Leak, FStructDeclarationPtr(Declaration), {});
+	}
+
+	FDualStructId GetSuper(const UStruct* Struct)
+	{
+		if (Struct->HasAnyCastFlags(CASTCLASS_UFunction))
+		{
+			return GUE.Structs.Function;
+		}
+		else if (Struct->HasAnyCastFlags(CASTCLASS_UClass))
+		{
+			FType Type = IndexStruct(Struct->GetClass());
+			return FDualStructId(GUE.Names.IndexDeclId(Type)); // e.g. UBlueprintGeneratedClass
+		}
+		else if (Struct->HasAnyCastFlags(CASTCLASS_UScriptStruct))
+		{
+			return GUE.Structs.ScriptStruct;
+		}
+		return GUE.Structs.Struct;
+	}
+
+	void BindAllMembers(const UStruct* Struct)
+	{
+		for (const FField* Child = Struct->ChildProperties; Child; Child = Child->Next)
+		{
+			BindMember(static_cast<const FProperty*>(Child));
+		}
+		if (const UClass* Class = Cast<UClass>(Struct))
+		{
+			// Todo: non-composable hack, declare CDO for UClass in the generic meta binding
+			FDualStructId Id = GUE.Structs.ObjectPtr;
+			Members.Emplace(FBindId(Id));
+			Specs.Emplace(FDeclId(Id));
+			Names.Emplace(GUE.Names.NameMember("ClassDefaultObject"));
+		}
+	}
+
+	void BindMember(const FProperty* Property)
+	{
+		FDualStructId Id = PropertyStructIdMeta(Property);
+		Members.Emplace(FBindId(Id));
+		Specs.Emplace(FDeclId(Id));
+		Names.Emplace(GUE.Names.NameMember(Property->GetFName()));
+	}
+
+	FStructDeclarationPtr Declare(FDeclId Super) const
+	{
+		return PlainProps::Declare({LowerCast(Owner), Super, 0, EMemberPresence::AllowSparse, Names, {Specs}});
+	}
+
+private:
+	const FBindId									Owner;
+
+	TArray<FMemberId, TInlineAllocator<64>>			Names;
+	TArray<FBindId, TInlineAllocator<64>>			Members;
+	TArray<FMemberSpec, TInlineAllocator<64>>		Specs;
+};
+
+void BindStructMeta(FBindId Id, const UStruct* Struct)
+{
+	if (!IsNativeOrTransientStruct(Struct))
+	{
+		FPropertyBinderMeta Binder(Id);
+		Binder.DeclareAndBind(Struct);
+	}
+}
+
+} // namespace PlainProps::UE
+
+////////////////////////////////////////////////////////////////////////////////////////////////
+
+namespace PlainProps
+{
+// Temporary way to add a fake /Script/CoreUObject scope to non-UEnums, not strictly required
+template<> inline FScopeId GetOuterScope<UE::FRuntimeIds, TTypename<EClassFlags>>()		{ return UE::GUE.Scopes.CoreUObject; }
+template<> inline FScopeId GetOuterScope<UE::FRuntimeIds, TTypename<EFunctionFlags>>()	{ return UE::GUE.Scopes.CoreUObject; }
+template<> inline FScopeId GetOuterScope<UE::FRuntimeIds, TTypename<EStructFlags>>()	{ return UE::GUE.Scopes.CoreUObject; }
+template<> inline FScopeId GetOuterScope<UE::FRuntimeIds, TTypename<EPropertyFlags>>()	{ return UE::GUE.Scopes.CoreUObject; }
+
+}
diff --git a/Engine/Plugins/Experimental/PlainPropsUObject/Source/Public/PlainPropsUObjectRuntime.h b/Engine/Plugins/Experimental/PlainPropsUObject/Source/Public/PlainPropsUObjectRuntime.h
--- a/Engine/Plugins/Experimental/PlainPropsUObject/Source/Public/PlainPropsUObjectRuntime.h
+++ b/Engine/Plugins/Experimental/PlainPropsUObject/Source/Public/PlainPropsUObjectRuntime.h
@@ -10,6 +10,7 @@
 #include "Containers/Map.h"
 #include "UObject/DynamicallyTypedValue.h"
 #include "UObject/NameTypes.h"
+#include "UObject/ObjectPtr.h"
 #include "UObject/ScriptDelegateFwd.h"
 
 struct FFieldPath;
@@ -81,6 +82,22 @@ struct FCommonEnumIds
 	FEnumId						PixelFormat;
 };
 
+struct FCommonPropertyStructIds
+{
+	explicit FCommonPropertyStructIds(const FCommonScopeIds& Scopes, TIdIndexer<FSensitiveName>& Names);
+
+	FDualStructId				Property;
+	FDualStructId				EnumProperty;
+	FDualStructId				BoolProperty;
+	FDualStructId				StructProperty;
+	FDualStructId				ClassProperty;
+	FDualStructId				ObjectPropertyBase;
+	FDualStructId				ArrayProperty;
+	FDualStructId				SetProperty;
+	FDualStructId				MapProperty;
+	FDualStructId				OptionalProperty;
+};
+
 struct FCommonStructIds
 {
 	explicit FCommonStructIds(const FCommonScopeIds& Scopes, TIdIndexer<FSensitiveName>& Names);
@@ -89,6 +106,10 @@ struct FCommonStructIds
 	FDualStructId				Text;
 	FDualStructId				Guid;
 	FDualStructId				FieldPath;
+	FDualStructId				Struct;
+	FDualStructId				Class;
+	FDualStructId				Function;
+	FDualStructId				ScriptStruct;
 	FDualStructId				SoftObjectPath;
 	FDualStructId				ClassPtr;
 	FDualStructId				ObjectPtr;
@@ -129,7 +150,15 @@ struct FCommonMemberIds
 struct FDynamicIds
 {
 	void			AddEnum(const UEnum* Enum, FEnumId Id)			{ Enums.Add(Enum, Id); }
-	void			AddStruct(const UStruct* Struct, FBindId Id)	{ Structs.Add(Struct, FDualStructId(Id)); }
+	void			AddStruct(const UStruct* Struct, FBindId Id)
+	{
+		// Don't store the new id when adding a duplicated schema for a nested schema struct a in custom binding.
+		// Todo: Can be cleaned up when interleaving custom binding members with schema members.
+		if (!Structs.Contains(Struct))
+		{
+			Structs.Add(Struct, FDualStructId(Id));
+		}
+	}
 
 	FEnumId			GetEnum(const UEnum* Enum) const				{ return Enums.FindChecked(Enum); }
 	FDualStructId	GetStruct(const UStruct* Struct) const			{ return Structs.FindChecked(Struct); }
@@ -151,6 +180,7 @@ struct FGlobals
 	FCommonScopeIds				Scopes;
 	FCommonEnumIds				EnumIds;
 	FCommonStructIds			Structs;
+	FCommonPropertyStructIds	Properties;
 	FCommonTypenameIds			Typenames;
 	FCommonMemberIds			Members;
 	FDynamicIds					DynamicIds;
@@ -340,6 +370,14 @@ namespace PlainProps
 template<Enumeration E> requires (requires { StaticEnum<E>(); })
 struct TCustomBind<TEnumAsByte<E>> { using Type = UE::TUEnumAsByteBinding<E>; };
 
+// Enable usage of TObjectPtr<T> as members in ctti reflected structs
+template<typename T> struct	TTypename<TObjectPtr<T>>	{ static_assert(!sizeof(T), "Unsupported type for TObjectPtr"); };
+template<> struct TTypename<TObjectPtr<UClass>>			{ inline static constexpr std::string_view DeclName = "ClassPtr"; };
+template<> struct TTypename<TObjectPtr<UObject>>		{ inline static constexpr std::string_view DeclName = "ObjectPtr"; };
+template<> inline FScopeId GetOuterScope<UE::FRuntimeIds, TTypename<TObjectPtr<UClass>>>()	{ return UE::GUE.Scopes.CoreUObject; }
+template<> inline FScopeId GetOuterScope<UE::FRuntimeIds, TTypename<TObjectPtr<UObject>>>()	{ return UE::GUE.Scopes.CoreUObject; }
+
+// custom bindings
 template<> struct TCustomBind<FFieldPath> { using Type = UE::FFieldPathBinding; };
 template<> struct TCustomBind<FScriptDelegate> { using Type = UE::FDelegateBinding; };
 template<> struct TCustomBind<FMulticastScriptDelegate> { using Type = UE::FMulticastInlineDelegateBinding; };
