PlainProps: Exclude code in PlainPropsUObjectMetaBindings.cpp causing compile errors on some platforms.

#rnx

#Synced-CL 52382190

diff --git a/Engine/Plugins/Experimental/PlainPropsUObject/Source/Private/PlainPropsUObjectMetaBindings.cpp b/Engine/Plugins/Experimental/PlainPropsUObject/Source/Private/PlainPropsUObjectMetaBindings.cpp
--- a/Engine/Plugins/Experimental/PlainPropsUObject/Source/Private/PlainPropsUObjectMetaBindings.cpp
+++ b/Engine/Plugins/Experimental/PlainPropsUObject/Source/Private/PlainPropsUObjectMetaBindings.cpp
@@ -18,9 +18,11 @@ UE_DEFINE_PRIVATE_MEMBER_PTR(TObjectPtr<UStruct>, GSuperStruct, UStruct, SuperSt
 PRAGMA_DISABLE_DEPRECATION_WARNINGS
 PP_DEFINE_PRIVATE_MEMBER_NAMESPACE_BEGIN(UClass_Private)
 PP_DEFINE_PRIVATE_MEMBER_PTR(UClass, ClassDefaultObject, TObjectPtr<UObject>);
+#if PP_OVERRIDE_OBJECT_LOADING
 #if USTRUCT_FAST_ISCHILDOF_IMPL == USTRUCT_ISCHILDOF_STRUCTARRAY
 PP_DEFINE_PRIVATE_MEMBER_PTR(UStruct, ReinitializeBaseChainArray, void());
 #endif
+#endif
 PP_DEFINE_PRIVATE_MEMBER_NAMESPACE_END()
 
 PP_DEFINE_PRIVATE_MEMBER_NAMESPACE_BEGIN(FProperty_Private)
@@ -301,8 +303,10 @@ void UStructBinding::Load(UStruct& Dst, FStructLoadView Src, ECustomLoadMethod M
 
 	FMemberLoader Members(Src);
 	LoadStruct(&(Dst.*GSuperStruct), Members.GrabStruct());
+#if PP_OVERRIDE_OBJECT_LOADING
 #if USTRUCT_FAST_ISCHILDOF_IMPL == USTRUCT_ISCHILDOF_STRUCTARRAY
 	(Dst.*_ReinitializeBaseChainArray)();
+#endif
 #endif
 	if (Members.HasMore())
 	{
