PlainProps: Minor cleanup for Meta and Blueprint bindings

#rnx
#Synced-CL 52382190

diff --git a/Engine/Plugins/Experimental/PlainPropsEngine/Source/Private/PlainPropsBlueprintBindings.cpp b/Engine/Plugins/Experimental/PlainPropsEngine/Source/Private/PlainPropsBlueprintBindings.cpp
--- a/Engine/Plugins/Experimental/PlainPropsEngine/Source/Private/PlainPropsBlueprintBindings.cpp
+++ b/Engine/Plugins/Experimental/PlainPropsEngine/Source/Private/PlainPropsBlueprintBindings.cpp
@@ -27,8 +27,6 @@ PP_REFLECT_STRUCT(, UEdGraphPin, void, PinId, PinName,                  SourceIn
 #endif
 
 
-// PP_REFLECT_ENUM(, EEdGraphPinDirection, EGPD_Input, EGPD_Output, EGPD_MAX);
-
 namespace PlainProps::UE::Blueprint
 {
 
@@ -282,7 +280,7 @@ void UEdGraphPinBinding::LoadPinReference(UEdGraphPin& RequestingPin, UEdGraphPi
 	LoadStruct(&Temp, Struct);
 	if (UEdGraphPin* Pin = Temp.Get())
 	{
-		PinRef = Temp.Get();
+		PinRef = Pin;
 	}
 	else
 	{
diff --git a/Engine/Plugins/Experimental/PlainPropsUObject/Source/Private/PlainPropsUObjectBindings.cpp b/Engine/Plugins/Experimental/PlainPropsUObject/Source/Private/PlainPropsUObjectBindings.cpp
--- a/Engine/Plugins/Experimental/PlainPropsUObject/Source/Private/PlainPropsUObjectBindings.cpp
+++ b/Engine/Plugins/Experimental/PlainPropsUObject/Source/Private/PlainPropsUObjectBindings.cpp
@@ -4329,8 +4329,8 @@ FStructDeclarationPtr FPropertyBinder::Declare() const
 
 ////////////////////////////////////////////////////////////////////////////////////////////////
 
-// Todo: Merge with SavePropertyMeta below
-FDualStructId PropertyStructIdMeta(const FProperty* Property)
+// Meta version with similar structure as FPropertyBinder::BindSingleProperty
+FDualStructId IndexPropertyMeta(const FField* Property)
 {
 	FName MemberTypename = Property->GetFName();
 	uint64 Flags = Property->GetCastFlags();
@@ -4370,87 +4370,10 @@ FDualStructId PropertyStructIdMeta(const FProperty* Property)
 			? GUE.Properties.SetProperty
 			: GUE.Properties.OptionalProperty;
 	}
-	return GUE.Properties.Property;
-}
-
-// Meta version similar to FPropertyBinder::BindSingleProperty
-FBuiltStruct* SavePropertyMeta(const FField* Field, const FSaveContext& Ctx)
-{
-	FName MemberTypename = Field->GetFName();
-	uint64 Flags = Field->GetCastFlags();
-	if (HasAny<LeafMask>(Flags))
-	{
-		if (HasAny<CASTCLASS_FNumericProperty>(Flags))
-		{
-			return SaveStruct(Field, GUE.Properties.Property, Ctx);
-		}
-		return HasAny<CASTCLASS_FEnumProperty>(Flags)
-			? SaveStruct(Field, GUE.Properties.Property, Ctx) // Todo: FEnumPropertyBinding
-			: SaveStruct(Field, GUE.Properties.BoolProperty, Ctx);
-	}
-	else if (HasAny<CASTCLASS_FStructProperty>(Flags))
-	{
-		return SaveStruct(Field, GUE.Properties.StructProperty, Ctx);
-	}
-	else if (HasAny<CASTCLASS_FClassProperty>(Flags))
-	{
-		return SaveStruct(Field, GUE.Properties.ClassProperty, Ctx);
-	}
-	else if (HasAny<CASTCLASS_FObjectPropertyBase>(Flags))
-	{
-		return SaveStruct(Field, GUE.Properties.ObjectPropertyBase, Ctx);
-	}
-	else if (HasAny<ContainerMask>(Flags))
-	{
-		if (HasAny<CASTCLASS_FArrayProperty>(Flags))
-		{
-			return SaveStruct(Field, GUE.Properties.Property, Ctx); // Todo: FArrayPropertyBinding
-		}
-		if (HasAny<CASTCLASS_FMapProperty>(Flags))
-		{
-			return SaveStruct(Field, GUE.Properties.Property, Ctx); // Todo: FMapPropertyBinding
-		}
-		return HasAny<CASTCLASS_FSetProperty>(Flags)
-			? SaveStruct(Field, GUE.Properties.Property, Ctx) // Todo: FSetPropertyBinding
-			: SaveStruct(Field, GUE.Properties.Property, Ctx); // Todo: FOptionalBinding
-	}
-	else if (HasAny<StringMask>(Flags))
-	{
-		return SaveStruct(Field, GUE.Properties.Property, Ctx); // Todo: FStringPropertyBinding
-	}
-	else if (HasAny<MiscMask>(Flags))
-	{
-		return SaveStruct(Field, GUE.Properties.Property, Ctx); // Todo: FInterfaceBinding and FDelegateBindings
-	}
-#if WITH_VERSE_VM
-	else if (HasAny<CASTCLASS_FVValueProperty | CASTCLASS_FVRestValueProperty>(Flags))
-	{
-		// return HasAny<CASTCLASS_FVValueProperty>(Flags)
-		// 	? BindVValue(static_cast<FVValueProperty*>(Property))
-		// 	: BindVRestValue(static_cast<FVRestValueProperty*>(Property)); 
-	}
-#else // Verse BPVM
-	else
-	{
-		if (MemberTypename == "VerseFunctionProperty") // Todo: FVerseFunctionProperty
-		{
-			return SaveStruct(Field, GUE.Properties.Property, Ctx);
-		}
-		else if (MemberTypename == "VerseDynamicProperty") // Todo: FVerseDynamicProperty
-		{
-			return SaveStruct(Field, GUE.Properties.Property, Ctx);
-		}
-		else if (MemberTypename == "ReferenceProperty") // Todo: FReferenceProperty
-		{
-			return SaveStruct(Field, GUE.Properties.Property, Ctx);
-		}
-	}
-#endif
-
 	// Todo: check for complete coverage after implementing all missing FProperty bindings.
 	// checkf(false, TEXT("Unrecognized class cast flags %llx in %s %s"),
 	// Flags, *MemberTypename.ToString(), *Field->GetName());
-	return SaveStruct(Field, GUE.Properties.Property, Ctx);
+	return GUE.Properties.Property;
 }
 
 ////////////////////////////////////////////////////////////////////////////////////////////////
diff --git a/Engine/Plugins/Experimental/PlainPropsUObject/Source/Private/PlainPropsUObjectBindingsInternal.h b/Engine/Plugins/Experimental/PlainPropsUObject/Source/Private/PlainPropsUObjectBindingsInternal.h
--- a/Engine/Plugins/Experimental/PlainPropsUObject/Source/Private/PlainPropsUObjectBindingsInternal.h
+++ b/Engine/Plugins/Experimental/PlainPropsUObject/Source/Private/PlainPropsUObjectBindingsInternal.h
@@ -23,8 +23,7 @@ const UStruct*	SkipEmptyBases(const UStruct* Struct);
 // Helpers for native UField and FProperty meta bindings
 FType			IndexStructMeta(const UStruct* Struct);
 void 			BindStructMeta(FBindId Id, const UStruct* Struct);
-FDualStructId 	PropertyStructIdMeta(const FProperty* Property);
-FBuiltStruct* 	SavePropertyMeta(const FField* Field, const FSaveContext& Ctx);
+FDualStructId 	IndexPropertyMeta(const FField* Property);
 
 } // namespace PlainProps::UE
 
diff --git a/Engine/Plugins/Experimental/PlainPropsUObject/Source/Private/PlainPropsUObjectMetaBindings.cpp b/Engine/Plugins/Experimental/PlainPropsUObject/Source/Private/PlainPropsUObjectMetaBindings.cpp
--- a/Engine/Plugins/Experimental/PlainPropsUObject/Source/Private/PlainPropsUObjectMetaBindings.cpp
+++ b/Engine/Plugins/Experimental/PlainPropsUObject/Source/Private/PlainPropsUObjectMetaBindings.cpp
@@ -319,7 +319,7 @@ void UStructBinding::Load(UStruct& Dst, FStructLoadView Src, ECustomLoadMethod M
 	if (Members.HasMore())
 	{
 		check(Members.PeekKind() == EMemberKind::Range);
-		check(Members.PeekName() == MemberIds[M::ChildPropertyClasses])
+		check(Members.PeekName() == MemberIds[M::ChildPropertyClasses]);
 		LoadChildPropertyClassesStatic(Dst, Members.GrabRange().AsStructs());
 	}
 	// UStruct::Serialize first serializes all FProperties, then calls link for all types but UClasses.
@@ -354,9 +354,11 @@ FTypedRange UClassBinding::SaveInterfacesStatic(FBindId Id, TConstArrayView<FImp
 
 void UClassBinding::LoadInterfacesStatic(TArray<FImplementedInterface>& Interfaces, FStructRangeLoadView Src)
 {
+	Interfaces.SetNum(static_cast<int32>(Src.Num()));
+	FImplementedInterface* DstIt = Interfaces.GetData();
 	for (FStructLoadView SrcIt : Src)
 	{
-		unimplemented();
+		LoadStruct(DstIt++, SrcIt);
 	}
 }
 
@@ -563,7 +565,7 @@ struct FStructBindingMeta final : ICustomBinding
 			const FBindId* BindId = &BindIds[0];
 			for (const FField* Child = Src.ChildProperties; Child; Child = Child->Next)
 			{
-				Dst.AddStruct(*MemberId++, *BindId++, SavePropertyMeta(Child, Ctx));
+				Dst.AddStruct(*MemberId++, *BindId++, SaveStruct(Child, IndexPropertyMeta(Child), Ctx));
 			}
 		}
 		if (const UClass* Class = Cast<UClass>(&Src))
@@ -600,7 +602,6 @@ struct FStructBindingMeta final : ICustomBinding
 			// Todo: The UClass interface arrray will also need to be loaded at this point in time
 			check(Members.HasMore());
 			LoadStruct(&((*Class).*_ClassDefaultObject), Members.GrabStruct());
-			FDualStructId Id = GUE.Structs.ObjectPtr;
 		}
 		check(!Members.HasMore());
 	}
@@ -664,7 +665,7 @@ public:
 
 	void BindMember(const FProperty* Property)
 	{
-		FDualStructId Id = PropertyStructIdMeta(Property);
+		FDualStructId Id = IndexPropertyMeta(Property);
 		Members.Emplace(FBindId(Id));
 		Specs.Emplace(FDeclId(Id));
 		Names.Emplace(GUE.Names.NameMember(Property->GetFName()));
