[PlainPropsUOjbect]
* Fixed shipping compile error

#Synced-CL 52415411

diff --git a/Engine/Plugins/Experimental/PlainPropsUObject/Source/Private/PlainPropsUObjectMetaBindings.cpp b/Engine/Plugins/Experimental/PlainPropsUObject/Source/Private/PlainPropsUObjectMetaBindings.cpp
--- a/Engine/Plugins/Experimental/PlainPropsUObject/Source/Private/PlainPropsUObjectMetaBindings.cpp
+++ b/Engine/Plugins/Experimental/PlainPropsUObject/Source/Private/PlainPropsUObjectMetaBindings.cpp
@@ -10,9 +10,11 @@
 #include "UObject/ObjectMacros.h"
 #include "UObject/ObjectPtr.h"
 #include "UObject/Package.h"
+#include "UObject/UnrealType.h"
 
 // Temp hacks. Long-term either add FProperty getters for ctor/dtor/hash function pointers 
 // and delegate APIs for non-intrusive serialization or integrate PlainProps into Core/CoreUObject
+
 UE_DEFINE_PRIVATE_MEMBER_PTR(TObjectPtr<UStruct>, GSuperStruct, UStruct, SuperStruct);
 
 PRAGMA_DISABLE_DEPRECATION_WARNINGS
