PlainProps: FInstancedStruct support
- Add load=test which creates memory assets for testing

#rnx

Test case with PP_WITH_ENGINE_HOOKS=1
-run=TestPlainProps -allowcommandletaudio -bind=source -batch=linker -load=maps

Before change, remove FInstancedStruct ignore:
Ignored 680 package diff(s) in 3567 package(s)
Detected 136 package diff(s) in 3567 package(s)

After:
Ignored 785 package diff(s) in 3567 package(s) (increase in diff we are serializing more data)
Detected 7 package diff(s) in 3567 package(s)

#Synced-CL 52466930

diff --git a/Engine/Plugins/Experimental/PlainPropsEngine/Source/Private/PlainPropsCommandlets.cpp b/Engine/Plugins/Experimental/PlainPropsEngine/Source/Private/PlainPropsCommandlets.cpp
--- a/Engine/Plugins/Experimental/PlainPropsEngine/Source/Private/PlainPropsCommandlets.cpp
+++ b/Engine/Plugins/Experimental/PlainPropsEngine/Source/Private/PlainPropsCommandlets.cpp
@@ -3,6 +3,8 @@
 #include "PlainPropsCommandlets.h"
 #include "PlainPropsEngineBindings.h"
 #include "PlainPropsRoundtripTest.h"
+#include "PlainPropsTestBindings.h"
+#include "PlainPropsTestAssets.h"
 #include "AssetRegistry/AssetData.h"
 #include "AssetRegistry/AssetRegistryModule.h"
 #include "Engine/World.h"
@@ -39,7 +41,6 @@ static PlainProps::UE::FLinkerDiffFilter LinkerDiffFilter
 
 		// Missing struct bindings
 		"InstancedPropertyBag",				// FInstancedPropertyBag in CoreUObject
-		"InstancedStruct",					// FInstancedStruct in CoreUObject
 		"MovieSceneFrameRange",				// FMovieSceneFrameRange in MovieScene,
 											// requires Math bindings for TRange<T> and TRangeBound<T>
 		"ShaderValueTypeHandle",			// FShaderValueTypeHandle in ComputeFramework
@@ -265,6 +266,10 @@ int32 UTestPlainPropsCommandlet::Main(const FString& Params)
 				AssetPath.LoadAsync({});
 			}
 		}
+		else if (AssetString == TEXT("test"))
+		{
+			Objects.Append(CreateTestAssets());
+		}
 		else
 		{
 			int32 CommaIndex;
@@ -313,6 +318,7 @@ int32 UTestPlainPropsCommandlet::Main(const FString& Params)
 		CustomBindAnimationTypes(Mode);
 		CustomBindNiagaraTypes(Mode);
 		CustomBindBlueprintTypes(Mode);
+		CustomBindTestTypes(Mode);
 	}
 
 	// Default roundtrip options
diff --git a/Engine/Plugins/Experimental/PlainPropsEngine/Source/Private/PlainPropsEngineTestUtils.cpp b/Engine/Plugins/Experimental/PlainPropsEngine/Source/Private/PlainPropsEngineTestUtils.cpp
--- a/Engine/Plugins/Experimental/PlainPropsEngine/Source/Private/PlainPropsEngineTestUtils.cpp
+++ b/Engine/Plugins/Experimental/PlainPropsEngine/Source/Private/PlainPropsEngineTestUtils.cpp
@@ -3,6 +3,7 @@
 #include "PlainPropsEngineTestUtils.h"
 #include "PlainPropsEngineBindings.h"
 #include "PlainPropsRoundtripTest.h"
+#include "PlainPropsTestBindings.h"
 
 namespace PlainProps::UE
 {
@@ -13,6 +14,7 @@ void BindAllTypes(EBindMode Mode, EBatchType BatchType)
 	CustomBindEngineTypes(Mode);
 	CustomBindAnimationTypes(Mode);
 	CustomBindNiagaraTypes(Mode);
+	CustomBindTestTypes(Mode);
 }
 
 } // namespace PlainProps::UE
\ No newline at end of file
diff --git a/Engine/Plugins/Experimental/PlainPropsEngine/Source/Private/PlainPropsTestAssets.cpp b/Engine/Plugins/Experimental/PlainPropsEngine/Source/Private/PlainPropsTestAssets.cpp
new file mode 100644
--- /dev/null
+++ b/Engine/Plugins/Experimental/PlainPropsEngine/Source/Private/PlainPropsTestAssets.cpp
@@ -0,0 +1,123 @@
+// Copyright Epic Games, Inc. All Rights Reserved.
+
+#include "PlainPropsTestAssets.h"
+
+#if WITH_TESTS && !defined(PLATFORM_COMPILER_IWYU)
+
+#include "PlainPropsTestTypes.h"
+#include "AssetRegistry/AssetRegistryModule.h"
+#include "Misc/PackagePath.h"
+#include "Serialization/LargeMemoryReader.h"
+#include "UObject/LinkerLoad.h"
+#include "UObject/Package.h"
+#include "UObject/SavePackage.h"
+#include "UObject/UObjectGlobals.h"
+#include "UObject/UObjectGlobalsInternal.h"
+#include "UObject/UObjectHash.h"
+#include "UObject/UObjectThreadContext.h"
+
+namespace PlainProps::UE
+{
+
+template<typename T>
+T* CreateTestAsset(const FString& AssetName, T* Template = nullptr)
+{
+	const FString PackageName = TEXT("/Test/") + AssetName;
+	UPackage* Package = CreatePackage(*PackageName);
+
+	T* Asset = NewObject<T>(Package, T::StaticClass(), *AssetName, RF_Public | RF_Standalone, Template);
+	check(Asset);
+	FAssetRegistryModule::AssetCreated(Asset);
+
+	Package->ClearPackageFlags(PKG_NewlyCreated);
+	Package->SetFlags(RF_Public | RF_WasLoaded);
+
+	return Asset;
+}
+
+TArray<UObject*> CreateTestAssets()
+{
+	TArray<UObject*> Assets;
+
+	{
+		UPlainPropsInstancedStructTestAsset* TestAsset = CreateTestAsset<UPlainPropsInstancedStructTestAsset>(TEXT("PlainPropsInstancedStructTestAssetModified"));
+		{
+			FPlainPropsInstancedStructTestSchemaA& InstancedStruct = TestAsset->StructA.GetMutable<FPlainPropsInstancedStructTestSchemaA>();
+			InstancedStruct.bValue = true;
+			InstancedStruct.X++;
+		}
+		
+		{
+			FPlainPropsInstancedStructTestSchemaB& InstancedStruct = TestAsset->StructB.GetMutable<FPlainPropsInstancedStructTestSchemaB>();
+			InstancedStruct.Count++;
+			InstancedStruct.Label = FName("Modified");
+		}
+
+		{
+			FPlainPropsInstancedStructTestCustomBound& InstancedStruct = TestAsset->StructC.GetMutable<FPlainPropsInstancedStructTestCustomBound>();
+			InstancedStruct.bValue = true;
+			InstancedStruct.Id++;
+		}
+
+		{
+			FPlainPropsInstancedStructTestSchemaA& InstancedStruct = TestAsset->MixedArray[0].GetMutable<FPlainPropsInstancedStructTestSchemaA>();
+			InstancedStruct.X+=1;
+			InstancedStruct.Y+=2;
+			InstancedStruct.Z+=3;
+		}
+		
+		{
+			FPlainPropsInstancedStructTestSchemaB& InstancedStruct = TestAsset->MixedArray[1].GetMutable<FPlainPropsInstancedStructTestSchemaB>();
+			InstancedStruct = FPlainPropsInstancedStructTestSchemaB();
+
+			InstancedStruct.Label = FName("Modified2");
+		}
+		
+		{
+			FPlainPropsInstancedStructTestCustomBound& InstancedStruct = TestAsset->MixedArray[2].GetMutable<FPlainPropsInstancedStructTestCustomBound>();
+			InstancedStruct.bValue = true;
+			InstancedStruct.Id++;
+		}
+
+		{
+			TestAsset->MixedArray[3].InitializeAs<FPlainPropsInstancedStructTestSchemaB>();
+			FPlainPropsInstancedStructTestSchemaB& InstancedStruct = TestAsset->MixedArray[3].GetMutable<FPlainPropsInstancedStructTestSchemaB>();
+			InstancedStruct.Label = FName("Modified3");
+		}
+
+		{
+			TestAsset->MixedArray.Add(FInstancedStruct::Make<FPlainPropsInstancedStructTestSchemaA>(FPlainPropsInstancedStructTestSchemaA({ {true}, 20.f, 21.f, 22.f })));
+		}
+		
+		{
+			FPlainPropsInstancedStructTestSchemaA& InstancedStruct = TestAsset->BaseStruct.GetMutable<FPlainPropsInstancedStructTestSchemaA>();
+			InstancedStruct.X+=2;
+			InstancedStruct.Y+=3;
+			InstancedStruct.Z+=4;
+		}
+
+		{
+			TestAsset->EmptyStruct.InitializeAs<FPlainPropsInstancedStructTestCustomBound>(FPlainPropsInstancedStructTestCustomBound{ { true }, 12 });
+		}
+
+		Assets.Add(TestAsset);
+	}
+
+	return Assets;
+}
+
+} // namespace PlainProps::UE
+
+#else
+
+namespace PlainProps::UE
+{
+
+TArray<UObject*> CreateTestAssets()
+{
+	return {};
+}
+
+}
+
+#endif // WITH_TESTS && !defined(PLATFORM_COMPILER_IWYU)
\ No newline at end of file
diff --git a/Engine/Plugins/Experimental/PlainPropsEngine/Source/Private/PlainPropsTestAssets.h b/Engine/Plugins/Experimental/PlainPropsEngine/Source/Private/PlainPropsTestAssets.h
new file mode 100644
--- /dev/null
+++ b/Engine/Plugins/Experimental/PlainPropsEngine/Source/Private/PlainPropsTestAssets.h
@@ -0,0 +1,14 @@
+// Copyright Epic Games, Inc. All Rights Reserved.
+
+#pragma once
+
+#include "Containers/Array.h"
+
+class UObject;
+
+namespace PlainProps::UE
+{
+
+TArray<UObject*> CreateTestAssets();
+
+} // namespace PlainProps::UE
diff --git a/Engine/Plugins/Experimental/PlainPropsEngine/Source/Private/PlainPropsTestBindings.cpp b/Engine/Plugins/Experimental/PlainPropsEngine/Source/Private/PlainPropsTestBindings.cpp
new file mode 100644
--- /dev/null
+++ b/Engine/Plugins/Experimental/PlainPropsEngine/Source/Private/PlainPropsTestBindings.cpp
@@ -0,0 +1,77 @@
+// Copyright Epic Games, Inc. All Rights Reserved.
+
+#include "PlainPropsTestBindings.h"
+
+#if WITH_TESTS && !defined(PLATFORM_COMPILER_IWYU)
+
+#include "PlainPropsTestTypes.h"
+#include "PlainPropsUObjectRuntime.h"
+
+namespace PlainProps::UE
+{
+
+namespace PlainPropsTestWithCustomBinding
+{
+
+/** 
+ *  Custom PlainProps binding for FPlainPropsInstancedStructTestCustomBound — serializes members explicitly by name.
+ *  Used to exercise the FInstancedStruct binding when the inner struct itself has a custom binding. 
+ * 
+ *  todo: once custom binding can forward to schema simplify this implementation to avoid having to keep it in sync to avoid diffs. 
+ */
+struct FCustomBinding : ICustomBinding
+{
+	using Type = FPlainPropsInstancedStructTestCustomBound;
+	
+	static constexpr int32 MemberCount = 2;
+	const FMemberId MemberIds[MemberCount];
+
+	template<class Ids>
+	FCustomBinding(TCustomSpecifier<Ids, MemberCount>& Spec)
+	: MemberIds{Ids::IndexMember("bValue"), Ids::IndexMember("Id")}
+	{
+		Spec.Members[0] = Specify<bool>();
+		Spec.Members[1] = Specify<int32>();
+	}
+
+	void Save(FMemberBuilder& Dst, const Type& Src, const Type* /*Default*/, const FSaveContext&) const
+	{
+		Dst.Add(MemberIds[0], Src.bValue);
+		Dst.Add(MemberIds[1], Src.Id);
+	}
+
+	void Load(Type& Dst, FStructLoadView Src, ECustomLoadMethod /*Method*/) const
+	{
+		FMemberLoader Members(Src);
+		Dst.bValue = Members.GrabLeaf().AsBool();
+		Dst.Id = Members.GrabLeaf().AsS32();
+	}
+
+	static bool Diff(const Type& A, const Type& B, const FBindContext&)
+	{
+		return !(A == B);
+	}
+};
+
+} // namespace PlainPropsTestWithCustomBinding
+
+void CustomBindTestTypes(EBindMode Mode)
+{
+	// Todo: Ownership / memory leak
+	new TScopedDefaultUStructBinding<FPlainPropsInstancedStructTestCustomBound, PlainPropsTestWithCustomBinding::FCustomBinding>();
+}
+
+} // namespace PlainProps::UE
+
+#else
+
+namespace PlainProps::UE
+{
+
+void CustomBindTestTypes(EBindMode)
+{
+}
+
+} // namespace PlainProps::UE
+
+#endif // WITH_TESTS && !defined(PLATFORM_COMPILER_IWYU)
\ No newline at end of file
diff --git a/Engine/Plugins/Experimental/PlainPropsEngine/Source/Private/PlainPropsTestBindings.h b/Engine/Plugins/Experimental/PlainPropsEngine/Source/Private/PlainPropsTestBindings.h
new file mode 100644
--- /dev/null
+++ b/Engine/Plugins/Experimental/PlainPropsEngine/Source/Private/PlainPropsTestBindings.h
@@ -0,0 +1,14 @@
+// Copyright Epic Games, Inc. All Rights Reserved.
+
+#pragma once
+
+#include "HAL/Platform.h"
+
+namespace PlainProps::UE
+{
+
+enum class EBindMode : uint8;
+
+void CustomBindTestTypes(EBindMode Mode);
+
+} // namespace PlainProps::UE
diff --git a/Engine/Plugins/Experimental/PlainPropsEngine/Source/Public/PlainPropsTestTypes.h b/Engine/Plugins/Experimental/PlainPropsEngine/Source/Public/PlainPropsTestTypes.h
new file mode 100644
--- /dev/null
+++ b/Engine/Plugins/Experimental/PlainPropsEngine/Source/Public/PlainPropsTestTypes.h
@@ -0,0 +1,111 @@
+// Copyright Epic Games, Inc. All Rights Reserved.
+
+#pragma once
+
+#if WITH_TESTS && !defined(PLATFORM_COMPILER_IWYU)
+
+#include "Engine/DataAsset.h"
+#include "StructUtils/InstancedStruct.h"
+
+#include "PlainPropsTestTypes.generated.h"
+
+USTRUCT()
+struct FPlainPropsInstancedStructTestSchemaBase
+{
+	GENERATED_BODY()
+
+	UPROPERTY(EditAnywhere, Category = "Default")
+	bool bValue = false;
+};
+
+USTRUCT()
+struct FPlainPropsInstancedStructTestSchemaA : public FPlainPropsInstancedStructTestSchemaBase
+{
+	GENERATED_BODY()
+
+	UPROPERTY(EditAnywhere, Category = "Default")
+	float X = 1.f;
+
+	UPROPERTY(EditAnywhere, Category = "Default")
+	float Y = 2.f;
+
+	UPROPERTY(EditAnywhere, Category = "Default")
+	float Z = 3.f;
+};
+
+USTRUCT()
+struct FPlainPropsInstancedStructTestSchemaB : public FPlainPropsInstancedStructTestSchemaBase
+{
+	GENERATED_BODY()
+
+	UPROPERTY(EditAnywhere, Category = "Default")
+	int32 Count = 4;
+
+	UPROPERTY(EditAnywhere, Category = "Default")
+	FName Label = FName("Five");
+};
+
+USTRUCT()
+struct FPlainPropsInstancedStructTestCustomBound : public FPlainPropsInstancedStructTestSchemaBase
+{
+	GENERATED_BODY()
+
+	UPROPERTY(EditAnywhere, Category = "Default")
+	int32 Id = 6;
+
+	// Used by custom binding
+	friend bool operator==(const FPlainPropsInstancedStructTestCustomBound& Lhs, const FPlainPropsInstancedStructTestCustomBound& Rhs)
+	{
+		return Lhs.Id == Rhs.Id && Lhs.bValue == Rhs.bValue;
+	}
+};
+
+/** Test asset for exercising FInstancedStruct serialization through the PlainProps custom binding. */
+UCLASS()
+class UPlainPropsInstancedStructTestAsset : public UDataAsset
+{
+	GENERATED_BODY()
+
+public:
+	UPlainPropsInstancedStructTestAsset()
+	{
+		StructA.InitializeAs<FPlainPropsInstancedStructTestSchemaA>();
+		StructB.InitializeAs<FPlainPropsInstancedStructTestSchemaB>();
+		StructC.InitializeAs<FPlainPropsInstancedStructTestCustomBound>();
+
+		MixedArray.Add(FInstancedStruct::Make<FPlainPropsInstancedStructTestSchemaA>(FPlainPropsInstancedStructTestSchemaA({{true}, 7.f, 8.f, 9.f})));
+		MixedArray.Add(FInstancedStruct::Make<FPlainPropsInstancedStructTestSchemaB>(FPlainPropsInstancedStructTestSchemaB{{false}, 10, FName("Eleven")}));
+		MixedArray.Add(FInstancedStruct::Make<FPlainPropsInstancedStructTestCustomBound>(FPlainPropsInstancedStructTestCustomBound{{true}, 12}));
+		MixedArray.Add(FInstancedStruct::Make<FPlainPropsInstancedStructTestSchemaA>());
+
+		// EmptyStruct intentionally left empty
+
+		BaseStruct.InitializeAs<FPlainPropsInstancedStructTestSchemaA>(FPlainPropsInstancedStructTestSchemaA{ { false }, 13.f, 14.f, 15.f });
+	}
+
+	// Single FInstancedStruct holding FPlainPropsInstancedStructTestSchemaA
+	UPROPERTY(EditAnywhere, Category = "Default")
+	FInstancedStruct StructA;
+
+	// Single FInstancedStruct holding FPlainPropsInstancedStructTestSchemaB
+	UPROPERTY(EditAnywhere, Category = "Default")
+	FInstancedStruct StructB;
+
+	// Single FInstancedStruct holding FPlainPropsInstancedStructTestCustomBound (with a custom binding)
+	UPROPERTY(EditAnywhere, Category = "Default")
+	FInstancedStruct StructC;
+
+	// Heterogeneous array: mix of FPlainPropsInstancedStructTestSchemaA, FPlainPropsInstancedStructTestSchemaB and FPlainPropsInstancedStructTestCustomBound entries
+	UPROPERTY(EditAnywhere, Category = "Default")
+	TArray<FInstancedStruct> MixedArray;
+
+	// Intentionally empty — exercises the null/empty case
+	UPROPERTY(EditAnywhere, Category = "Default")
+	FInstancedStruct EmptyStruct;
+
+	// Single TInstancedStruct<FPlainPropsInstancedStructTestSchemaBase> holding FPlainPropsInstancedStructTestSchemaA
+	UPROPERTY(EditAnywhere, Category = "Default")
+	TInstancedStruct<FPlainPropsInstancedStructTestSchemaBase> BaseStruct;
+};
+
+#endif // WITH_TESTS && !defined(PLATFORM_COMPILER_IWYU)
\ No newline at end of file
diff --git a/Engine/Plugins/Experimental/PlainPropsUObject/Source/Private/PlainPropsRoundtripTest.cpp b/Engine/Plugins/Experimental/PlainPropsUObject/Source/Private/PlainPropsRoundtripTest.cpp
--- a/Engine/Plugins/Experimental/PlainPropsUObject/Source/Private/PlainPropsRoundtripTest.cpp
+++ b/Engine/Plugins/Experimental/PlainPropsUObject/Source/Private/PlainPropsRoundtripTest.cpp
@@ -26,6 +26,7 @@
 #include "Serialization/LargeMemoryReader.h"
 #include "Serialization/MemoryReader.h"
 #include "Serialization/MemoryWriter.h"
+#include "Serialization/PackageStore.h"
 #include "Serialization/UnversionedPropertySerializationInternal.h"
 #include "Templates/UniquePtr.h"
 #include "Templates/UnrealTemplate.h"
@@ -1378,6 +1379,13 @@ static UPackage* UnpinPackage(UPackage* Package)
 	return Package;
 }
 
+static FString GetPackageSaveFilename(UPackage* Package)
+{
+	const FString& Extension = Package->ContainsMap() ? FPackageName::GetMapPackageExtension() : FPackageName::GetAssetPackageExtension();
+	FString SaveFilename = FPaths::Combine(*FPaths::ProjectSavedDir(), TEXT("PlainProps"), Package->GetName() + Extension);
+	return SaveFilename;
+}
+
 struct FPlainPropsLinkerSandbox
 {
 	static constexpr uint32 LoadFlags = (LOAD_Async | LOAD_NoVerify | LOAD_SkipLoadImportedPackages);
@@ -1439,10 +1447,60 @@ struct FPlainPropsLinkerSandbox
 		FPersistentBatchLoader BatchLoader;
 	};
 
+	class FPackageStoreBackend final : public IPackageStoreBackend
+	{
+		struct FEntry
+		{
+			FName PackageName;
+			EPackageExtension Extension;
+		};
+		TMap<FName, FEntry> Packages;
+
+	public:
+		void AddPackage(FName PackageName, const FPackagePath& PackagePath)
+		{
+			Packages.Add(PackagePath.GetPackageFName(), { PackageName, PackagePath.GetHeaderExtension() });
+		}
+
+		void ClearPackages()
+		{
+			Packages.Reset();
+		}
+
+		virtual EPackageLoader GetSupportedLoaders() override
+		{
+			return EPackageLoader::LinkerLoad;
+		}
+
+		virtual void OnMounted(TSharedRef<const FPackageStoreBackendContext>) override {}
+		virtual void BeginRead() override {}
+		virtual void EndRead() override {}
+
+		virtual EPackageStoreEntryStatus GetPackageStoreEntry(FPackageId PackageId, FName PackageName, FPackageStoreEntry& OutEntry) override
+		{
+			if (const FEntry* Entry = Packages.Find(PackageName))
+			{
+				OutEntry.LoaderType = EPackageLoader::LinkerLoad;
+				OutEntry.PackageExtension = Entry->Extension;
+#if WITH_EDITORONLY_DATA
+				OutEntry.LinkerLoadCaseCorrectedPackageName = Entry->PackageName;
+#endif
+				return EPackageStoreEntryStatus::Ok;
+			}
+			return EPackageStoreEntryStatus::Missing;
+		}
+
+		virtual bool GetPackageRedirectInfo(FPackageId, FName&, FPackageId&) override
+		{
+			return false;
+		}
+	};
+
 	static TMap<FName, FSavedPackage> SavedPackages;
 	static TMap<FName, TUniquePtr<FMountedLinker>> MountedBaselinePackages;
 	static TMap<FName, TUniquePtr<FMountedLinkerBatch>> MountedPlainpropsPackages;
 	static TArray<const UClass*> PendingClasses;
+	static TSharedPtr<FPackageStoreBackend> PackageStoreBackend;
 
 #if WITH_EDITOR
 	static void OnBlueprintCDOCompiled(UObject* InObject, const FObjectPostCDOCompiledContext&)
@@ -1502,7 +1560,9 @@ struct FPlainPropsLinkerSandbox
 
 		FSavedPackage& SavedPackage = SavedPackages.FindOrAdd(PackageName);
 		SavedPackage.PackageName = PackageName;
-		SavedPackage.PackagePath = LinkerSave.LinkerRoot->GetLoadedPath();
+		SavedPackage.PackagePath = !LinkerSave.LinkerRoot->GetLoadedPath().IsEmpty() ?
+			LinkerSave.LinkerRoot->GetLoadedPath() :
+			FPackagePath::FromLocalPath(GetPackageSaveFilename(LinkerSave.LinkerRoot));
 		SavedPackage.PlainPropsData = MoveTemp(StableData);
 		SavedPackage.Texts = Properties.GrabTexts();
 	}
@@ -1538,6 +1598,7 @@ struct FPlainPropsLinkerSandbox
 			TUniquePtr<FMountedLinker> MountedInstance(new FMountedLinker(*SavedData, NewPackageName));
 			UPackage* Out = MountedInstance->Package;
 			MountedBaselinePackages.Emplace(NewPackageName, MoveTemp(MountedInstance));
+			PackageStoreBackend->AddPackage(NewPackageName, SavedData->PackagePath);
 			return Out;
 		}
 		return nullptr;
@@ -1556,16 +1617,28 @@ struct FPlainPropsLinkerSandbox
 			TUniquePtr<FMountedLinkerBatch> MountedInstance(new FMountedLinkerBatch(*SavedData, NewPackageName));
 			UPackage* Out = MountedInstance->MountedLinker.Package;
 			MountedPlainpropsPackages.Add(NewPackageName, MoveTemp(MountedInstance));
+			PackageStoreBackend->AddPackage(NewPackageName, SavedData->PackagePath);
 			return Out;
 		}
 		return nullptr;
 	}
 
+	static void Initialize(uint32 NumPackages)
+	{
+		SavedPackages.Reserve(NumPackages);
+		MountedBaselinePackages.Reserve(NumPackages);
+		MountedPlainpropsPackages.Reserve(NumPackages);
+		
+		PackageStoreBackend = MakeShared<FPackageStoreBackend>();
+		FPackageStore::Get().Mount(PackageStoreBackend.ToSharedRef(), /*Priority=*/10);
+	}
+
 	static void Shutdown()
 	{
 		MountedBaselinePackages.Empty();
 		MountedPlainpropsPackages.Empty();
 		SavedPackages.Empty();
+		PackageStoreBackend->ClearPackages();
 	}
 
 	static UPackage* LoadBaselinePackage(UPackage* Package)
@@ -1718,6 +1791,7 @@ TMap<FName, FPlainPropsLinkerSandbox::FSavedPackage> FPlainPropsLinkerSandbox::S
 TMap<FName, TUniquePtr<FPlainPropsLinkerSandbox::FMountedLinkerBatch>> FPlainPropsLinkerSandbox::MountedPlainpropsPackages;
 TMap<FName, TUniquePtr<FPlainPropsLinkerSandbox::FMountedLinker>> FPlainPropsLinkerSandbox::MountedBaselinePackages;
 TArray<const UClass*> FPlainPropsLinkerSandbox::PendingClasses;
+TSharedPtr<FPlainPropsLinkerSandbox::FPackageStoreBackend> FPlainPropsLinkerSandbox::PackageStoreBackend;
 
 #endif
 
@@ -2221,9 +2295,7 @@ int32 RoundtripViaLinkerBatch(
 	}
 	int32 NumPackages = SavedPackages.Num();
 
-	FPlainPropsLinkerSandbox::SavedPackages.Reserve(NumPackages);
-	FPlainPropsLinkerSandbox::MountedBaselinePackages.Reserve(NumPackages);
-	FPlainPropsLinkerSandbox::MountedPlainpropsPackages.Reserve(NumPackages);
+	FPlainPropsLinkerSandbox::Initialize(NumPackages);
 
 	UE_LOGFMT(LogPlainPropsUObject, Display, "Saving {SaveNum} packages...", NumPackages);
 
diff --git a/Engine/Plugins/Experimental/PlainPropsUObject/Source/Private/PlainPropsUObjectBindings.cpp b/Engine/Plugins/Experimental/PlainPropsUObject/Source/Private/PlainPropsUObjectBindings.cpp
--- a/Engine/Plugins/Experimental/PlainPropsUObject/Source/Private/PlainPropsUObjectBindings.cpp
+++ b/Engine/Plugins/Experimental/PlainPropsUObject/Source/Private/PlainPropsUObjectBindings.cpp
@@ -16,6 +16,7 @@
 #include "Logging/StructuredLog.h"
 #include "Math/PreciseFP.h"
 #include "Misc/DefinePrivateMemberPtr.h"
+#include "StructUtils/PropertyBag.h"
 #include "StructUtils/UserDefinedStruct.h"
 #include "Templates/SharedPointer.h"
 #include "Serialization/ArchiveSerializedPropertyChain.h"
@@ -342,6 +343,7 @@ FCommonStructIds::FCommonStructIds(const FCommonScopeIds& Scopes, TIdIndexer<FSe
 , DynamicallyTypedValue(	Names.IndexStruct({Scopes.CoreUObject, Names.MakeTypename("DynamicallyTypedValue")}))
 , ReferencePropertyValue(	Names.IndexStruct({Scopes.CoreUObject, Names.MakeTypename("ReferencePropertyValue")}))
 , PropertyBag(				Names.IndexStruct({Scopes.CoreUObject, Names.MakeTypename("PropertyBag")}))
+, InstancedStruct(			Names.IndexStruct({Scopes.CoreUObject, Names.MakeTypename("InstancedStruct")}))
 {}
 
 FCommonMemberIds::FCommonMemberIds(TIdIndexer<FSensitiveName>& Names)
@@ -964,6 +966,91 @@ bool FReferencePropertyBinding::Diff(const FReferencePropertyValue& A, const FRe
 
 ////////////////////////////////////////////////////////////////////////////////////////////////
 
+FInstancedStructBinding::FInstancedStructBinding(TPropertySpecifier<2>& Spec)
+	: MemberIds{ GUE.Names.NameMember("ScriptStruct"), GUE.Names.NameMember("ScriptMemory")}
+{
+	Spec.Members[0] = FDeclId(GUE.Structs.ObjectPtr);
+	Spec.Members[1] = SpecDynamicStruct;
+}
+
+void FInstancedStructBinding::Save(FMemberBuilder& Dst, const FInstancedStruct& Src, const FInstancedStruct* Default, const FSaveContext& Ctx) const
+{
+	const TObjectPtr<const UScriptStruct> StructRef(Src.GetScriptStruct());
+	
+	// todo: remove once we have UPropertyBag support
+	if (StructRef && StructRef->IsA<UPropertyBag>())
+	{
+		return;
+	}
+
+	Dst.AddStruct(MemberIds[0], GUE.Structs.ObjectPtr, SaveStruct(&StructRef, GUE.Structs.ObjectPtr, Ctx));
+
+	if (!StructRef)
+	{
+		return;
+	}
+
+	const FBindId BindId = GUE.DynamicIds.GetStruct(StructRef);
+	if (Default && Default->GetScriptStruct() == StructRef)
+	{
+		Dst.AddStruct(MemberIds[1], BindId, SaveStructDelta(Src.GetMemory(), FBaseline(Default->GetMemory()), BindId, Ctx));
+	}
+	else 
+	{
+		Dst.AddStruct(MemberIds[1], BindId, SaveStructDelta(Src.GetMemory(), GUE.Defaults.Get(BindId), BindId, Ctx));
+	}
+}
+
+void FInstancedStructBinding::Load(FInstancedStruct& Dst, FStructLoadView Src, ECustomLoadMethod Method) const
+{
+	FMemberLoader It(Src);
+
+	// todo: remove once we have UPropertyBag support
+	if (!It.HasMore())
+	{
+		return;
+	}
+
+	TObjectPtr<UObject> StructRef;
+	LoadStruct(&StructRef, It.GrabStruct());
+	
+	const UScriptStruct* ScriptStruct = Cast<UScriptStruct>(StructRef);
+	if (!ScriptStruct)
+	{
+		Dst.Reset();
+		return;
+	}
+
+	if (Dst.GetScriptStruct() != ScriptStruct)
+	{
+		Dst.InitializeAs(ScriptStruct);
+	}
+
+	LoadStruct(Dst.GetMutableMemory(), It.GrabStruct());
+}
+
+bool FInstancedStructBinding::Diff(const FInstancedStruct& A, const FInstancedStruct& B, const FBindContext& Ctx)
+{
+	if (A.GetScriptStruct() != B.GetScriptStruct())
+	{
+		return true;
+	}
+	else if (!A.GetScriptStruct())
+	{
+		return false;
+	}
+	// todo: remove once we have UPropertyBag support
+	else if (A.GetScriptStruct()->IsA<UPropertyBag>())
+	{
+		return false;
+	}
+
+	const FBindId BindId = GUE.DynamicIds.GetStruct(A.GetScriptStruct());
+	return DiffStructs(A.GetMemory(), B.GetMemory(), BindId, Ctx);
+}
+
+////////////////////////////////////////////////////////////////////////////////////////////////
+
 struct FInterfaceBinding final : ICustomBinding
 {
 	explicit FInterfaceBinding(UClass* Class) : InterfaceClass(Class) {}
@@ -2227,8 +2314,6 @@ bool DiffMap(const BindingType& Binding, const ScriptType& A, const ScriptType&
 				return true;
 			}
 
-			// Temp comment: Fails for Lyra -loadmaps since FInstancedStruct is unbound,
-			//				 so all instanceed structs are considered equal
 			checkSlow(Binding.HasItem(B, PairA));
 		}
 	}
@@ -4610,6 +4695,7 @@ struct FPlainPropsBindingScope
 		, LinearColor()
 		, FieldPath(GUE.Structs.FieldPath)
 		, Delegate(GUE.Structs.Delegate)
+		, InstancedStruct(GUE.Structs.InstancedStruct)
 		, InlineMulticast({ GUE.Structs.MulticastInlineDelegate, GUE.Structs.MulticastDelegate })
 		, VerseFunction(GUE.Structs.VerseFunction)
 		, DynamicallyTypedValue(GUE.Structs.DynamicallyTypedValue)
@@ -4635,6 +4721,7 @@ struct FPlainPropsBindingScope
 		GUE.Defaults.BindZeroes<FGuid>(Guid.BindId);
 		GUE.Defaults.BindZeroes<FColor>(Color.BindId);
 		GUE.Defaults.BindZeroes<FLinearColor>(LinearColor.BindId);
+		GUE.Defaults.BindZeroes<FInstancedStruct>(InstancedStruct.BindId);
 	}
 
 	~FPlainPropsBindingScope()
@@ -4643,6 +4730,7 @@ struct FPlainPropsBindingScope
 		GUE.Defaults.Drop(Guid.BindId);
 		GUE.Defaults.Drop(Color.BindId);
 		GUE.Defaults.Drop(LinearColor.BindId);
+		GUE.Defaults.Drop(InstancedStruct.BindId);
 	}
 
 	TScopedStructBinding<FTransform, FDefaultRuntime> Transform;
@@ -4651,6 +4739,7 @@ struct FPlainPropsBindingScope
 	TScopedStructBinding<FLinearColor, FDefaultRuntime> LinearColor;
 	TScopedStructBinding<FFieldPath, FDefaultRuntime> FieldPath;
 	TScopedStructBinding<FScriptDelegate, FDefaultRuntime> Delegate;
+	TScopedStructBinding<FInstancedStruct, FDefaultRuntime> InstancedStruct;
 	// MulticastDelegate declaration is shared with MulticastSparseDelegate
 	TScopedStructBinding<FMulticastScriptDelegate, FDefaultRuntime> InlineMulticast;
 	// Verse
diff --git a/Engine/Plugins/Experimental/PlainPropsUObject/Source/Private/PlainPropsUObjectDiff.cpp b/Engine/Plugins/Experimental/PlainPropsUObject/Source/Private/PlainPropsUObjectDiff.cpp
--- a/Engine/Plugins/Experimental/PlainPropsUObject/Source/Private/PlainPropsUObjectDiff.cpp
+++ b/Engine/Plugins/Experimental/PlainPropsUObject/Source/Private/PlainPropsUObjectDiff.cpp
@@ -9,6 +9,7 @@
 #include "UObject/UnrealType.h"
 #include "UObject/OverridableManager.h"
 #include "Containers/StringFwd.h"
+#include "StructUtils/InstancedStruct.h"
 
 // Temp hacks for non-intrusive prototype
 PP_DEFINE_PRIVATE_MEMBER_NAMESPACE_BEGIN(FFieldPath_Private)
@@ -223,6 +224,26 @@ bool DiffObjectReferences(const UObject* A, const UObject* B, FDiffObjectContext
 	return DiffStructs(A, B, Class, Ctx);
 }
 
+bool DiffInstancedStructs(const FInstancedStruct* A, const FInstancedStruct* B, const FStructProperty* Property, FDiffObjectContext& Ctx)
+{
+	if (A == B)
+	{
+		return false;
+	}
+
+	if (!A || !B || A->GetScriptStruct() != B->GetScriptStruct())
+	{
+		Ctx.Diffs.Emplace(EDiffObjectNodeType::Property, A, B, Property);
+		return true;
+	}
+	else if (!A->GetScriptStruct())
+	{
+		return false;
+	}
+
+	return DiffStructs(A->GetMemory(), B->GetMemory(), A->GetScriptStruct(), Ctx);
+}
+
 bool DiffSoftObjectPaths(const FSoftObjectPath* A, const FSoftObjectPath* B, const FStructProperty* Property, FDiffObjectContext& Ctx)
 {
 	if (A == B)
@@ -303,6 +324,12 @@ bool DiffStructProperties(const void* A, const void* B, const FStructProperty* P
 		return DiffSoftObjectPaths(static_cast<const FSoftObjectPath*>(A), static_cast<const FSoftObjectPath*>(B), Property, Ctx);
 	}
 
+	static const FName InstancedStruct{"InstancedStruct"};
+	if (ScriptStruct->GetFName() == InstancedStruct)
+	{
+		return DiffInstancedStructs(static_cast<const FInstancedStruct*>(A), static_cast<const FInstancedStruct*>(B), Property, Ctx);
+	}
+
 	// see ScriptStruct->CompareScriptStruct(A, B, /*PortFlags*/ 0);
 	if (ScriptStruct->StructFlags & STRUCT_IdenticalNative)
 	{
diff --git a/Engine/Plugins/Experimental/PlainPropsUObject/Source/Public/PlainPropsUObjectRuntime.h b/Engine/Plugins/Experimental/PlainPropsUObject/Source/Public/PlainPropsUObjectRuntime.h
--- a/Engine/Plugins/Experimental/PlainPropsUObject/Source/Public/PlainPropsUObjectRuntime.h
+++ b/Engine/Plugins/Experimental/PlainPropsUObject/Source/Public/PlainPropsUObjectRuntime.h
@@ -14,6 +14,7 @@
 #include "UObject/ScriptDelegateFwd.h"
 
 struct FFieldPath;
+struct FInstancedStruct;
 struct FReferencePropertyValue;
 struct FVerseFunction;
 
@@ -125,6 +126,7 @@ struct FCommonStructIds
 	FDualStructId				DynamicallyTypedValue;
 	FDualStructId				ReferencePropertyValue;
 	FDualStructId				PropertyBag;
+	FDualStructId				InstancedStruct;
 };
 
 struct FCommonMemberIds
@@ -363,6 +365,18 @@ struct FReferencePropertyBinding : ICustomBinding
 	static bool		Diff(const Type& A, const Type& B, const FBindContext&);
 };
 
+struct FInstancedStructBinding : ICustomBinding
+{
+	using Type = FInstancedStruct;
+	const FMemberId MemberIds[2];
+
+	FInstancedStructBinding(TPropertySpecifier<2>& Spec);
+
+	void			Save(FMemberBuilder& Dst, const Type& Src, const Type* Default, const FSaveContext&) const;
+	void			Load(Type& Dst, FStructLoadView Src, ECustomLoadMethod Method) const;
+	static bool		Diff(const Type& A, const Type& B, const FBindContext&);
+};
+
 } // namespace PlainProps::UE
 
 namespace PlainProps
@@ -384,6 +398,7 @@ template<> struct TCustomBind<FMulticastScriptDelegate> { using Type = UE::FMult
 template<> struct TCustomBind<::UE::FDynamicallyTypedValue> { using Type = UE::FDynamicallyTypedValueBinding; };
 template<> struct TCustomBind<FReferencePropertyValue> { using Type = UE::FReferencePropertyBinding; };
 template<> struct TCustomBind<FVerseFunction> { using Type = UE::FVerseFunctionBinding; };
+template<> struct TCustomBind<FInstancedStruct> { using Type = UE::FInstancedStructBinding; };
 
 // Temporary way to tie certain Core types to /Script/CoreUObject scope
 template<> inline FScopeId GetOuterScope<UE::FRuntimeIds, TTypename<FName>>()			{ return UE::GUE.Scopes.CoreUObject; }
