PlainProps: Add property flags metadata bindings
Fix saving and restoring of overridabble serialization state to be consistent with TPS beside where a replace is equivalent of all properties being modified in TPS which should be ignored at the diff level

#Synced-CL 52480300

diff --git a/Engine/Plugins/Experimental/PlainPropsUObject/Source/Private/PlainPropsRestoreOverridable.cpp b/Engine/Plugins/Experimental/PlainPropsUObject/Source/Private/PlainPropsRestoreOverridable.cpp
--- a/Engine/Plugins/Experimental/PlainPropsUObject/Source/Private/PlainPropsRestoreOverridable.cpp
+++ b/Engine/Plugins/Experimental/PlainPropsUObject/Source/Private/PlainPropsRestoreOverridable.cpp
@@ -1,6 +1,7 @@
 // Copyright Epic Games, Inc. All Rights Reserved.
 
 #include "PlainPropsRestoreOverridableInternal.h"
+#include "PlainPropsInternalPrivateMemberPtr.h"
 #include "PlainPropsBind.h"
 #include "PlainPropsUObjectRuntime.h"
 #include "Containers/Array.h"
@@ -15,6 +16,18 @@ namespace PlainProps::UE
 // Todo: This should be moved to an appropriate context that we can pass around.
 thread_local FOverriddenPropertySet* TlsRestoreOverrides = nullptr;
 
+PP_DEFINE_PRIVATE_MEMBER_NAMESPACE_BEGIN(FOverriddenPropertyNode_Private)
+PP_DEFINE_PRIVATE_MEMBER_PTR(FOverriddenPropertyNode, Operation, EOverriddenPropertyOperation);
+PP_DEFINE_PRIVATE_MEMBER_PTR(FOverriddenPropertyNode, SubPropertyNodes, TArray<FOverriddenPropertyNode>);
+PP_DEFINE_PRIVATE_MEMBER_NAMESPACE_END()
+
+void ResetOverriddenPropertyNode(FOverriddenPropertyNode& Node, EOverriddenPropertyOperation Op, TArray<FOverriddenPropertyNode> SubPropertyNodes)
+{
+	using namespace FOverriddenPropertyNode_Private;
+	Node.*_Operation = Op;
+	Node.*_SubPropertyNodes = MoveTemp(SubPropertyNodes);
+}
+
 // Determines whether an override operation is Modified or Replaced. A struct member is considered 
 // replaced if all of its members have been serialized. 
 // Note: in FOverriddenPropertySet::NotifyPropertyChange we consider an edited property as
@@ -24,8 +37,18 @@ static EOverriddenPropertyOperation DetermineOverrideOperation(FStructView Seria
 {
 	FMemberReader It(SerializedStruct);
 
+	// Handle super struct separately
+	if (It.HasSuper())
+	{
+		EOverriddenPropertyOperation SubOp = DetermineOverrideOperation(It.GrabStruct(), Ctx);
+		if (SubOp == EOverriddenPropertyOperation::Modified)
+		{
+			// When the sub-operation is Modified, then this member is not fully
+			// replaced and more sub-operations will be nested under it.
+			return EOverriddenPropertyOperation::Modified;
+		}
+	}
 	uint16 NumSerializedMembers = 0;
-
 	while (It.HasMore())
 	{
 		if (It.PeekKind() == EMemberKind::Struct)
@@ -42,7 +65,6 @@ static EOverriddenPropertyOperation DetermineOverrideOperation(FStructView Seria
 		{
 			It.SkipMember();
 		}
-
 		++NumSerializedMembers;
 	}
 
@@ -62,6 +84,12 @@ static EOverriddenPropertyOperation DetermineOverrideOperation(FStructView Seria
 	{
 		NumDeclMembers = CustomDecl->NumMembers;
 	}
+	// if there's no declaration this is a primitive type bound as a struct i.e object pointer, consider it replaced
+	else
+	{
+		ensure(NumSerializedMembers >= NumDeclMembers);
+		NumDeclMembers = NumSerializedMembers;
+	}
 
 	return (NumSerializedMembers == NumDeclMembers) ? EOverriddenPropertyOperation::Replace : EOverriddenPropertyOperation::Modified;
 }
@@ -133,6 +161,12 @@ static void RestoreStruct(const UStruct* Struct,
 
 	while (It.HasMore())
 	{
+		if (It.HasSuper())
+		{
+			FStructView Super = It.GrabStruct();
+			RestoreStruct(Struct, Super, Overrides, PropertyChain, Ctx);
+			continue;
+		}
 		if (FOptionalMemberId Name = It.PeekName())
 		{
 			FMemberType Type = It.PeekType();
@@ -148,11 +182,58 @@ static void RestoreStruct(const UStruct* Struct,
 				continue;
 			}
 		}
-
 		It.SkipMember();
 	}
 }
 
+void FlattenModifiedOverrides(FOverriddenPropertyNode& Root, FStructSchemaId Id, const FRestoreContext& Ctx)
+{
+	if (Root.GetOperation() == EOverriddenPropertyOperation::Modified)
+	{
+		FBindId BindId = FBindId(reinterpret_cast<const FDeclId*>(Ctx.RuntimeIds.GetData())[Id.Idx]);
+		uint32 NumDeclMembers = 0;
+		uint32 OverriddenCount = 0;
+		
+		FOptionalDeclId DeclId = FOptionalDeclId(LowerCast(BindId));
+		while (DeclId)
+		{
+			const FStructDeclaration* Decl = nullptr;
+			BindId = UpCast(DeclId.Get());
+			if (!Ctx.Schemas.FindStruct(BindId, Decl))
+			{
+				Decl = Ctx.Customs.FindDeclaration(BindId);
+			}
+			checkf(Decl, TEXT("'%s' declaration not found"), *Ctx.Schemas.GetDebug().Print(BindId));
+			for (EPropertyFlags MemberFlag : Ctx.Metadatas.GetMemberFlags(BindId))
+			{
+				OverriddenCount += !!(MemberFlag & EPropertyFlags::CPF_ExperimentalAlwaysOverriden);
+			}
+			if (Decl)
+			{
+				NumDeclMembers += Decl->NumMembers;
+				DeclId = Decl->Super;
+			}
+			else
+			{
+				DeclId = NoId;
+			}
+		}
+		
+		for (const FOverriddenPropertyNode& SubNodes : Root.GetSubPropertyNodes())
+		{
+			if (SubNodes.GetOperation() == EOverriddenPropertyOperation::Replace)
+			{
+				++OverriddenCount;
+			}
+		}
+
+		if (NumDeclMembers == OverriddenCount)
+		{
+			ResetOverriddenPropertyNode(Root, EOverriddenPropertyOperation::Replace, {});
+		}
+	}
+}
+
 void RestoreStructOverrides(const UStruct* Struct, FStructView StructView, 
 							FOverriddenPropertySet& Overrides, 
 							const FRestoreContext& Ctx)
@@ -160,6 +241,7 @@ void RestoreStructOverrides(const UStruct* Struct, FStructView StructView,
 	FArchiveSerializedPropertyChain PropertyChain;
 
 	RestoreStruct(Struct, StructView, Overrides, PropertyChain, Ctx);
+	FlattenModifiedOverrides(Overrides.GetRootOverriddenPropertyNode(), StructView.Schema.Id, Ctx);
 }
 
 FScopeRestoreOverrides::FScopeRestoreOverrides(FOverriddenPropertySet* OverriddenProperties)
diff --git a/Engine/Plugins/Experimental/PlainPropsUObject/Source/Private/PlainPropsRestoreOverridableInternal.h b/Engine/Plugins/Experimental/PlainPropsUObject/Source/Private/PlainPropsRestoreOverridableInternal.h
--- a/Engine/Plugins/Experimental/PlainPropsUObject/Source/Private/PlainPropsRestoreOverridableInternal.h
+++ b/Engine/Plugins/Experimental/PlainPropsUObject/Source/Private/PlainPropsRestoreOverridableInternal.h
@@ -6,6 +6,8 @@
 #include "PlainPropsTypes.h"
 
 struct FOverriddenPropertySet;
+struct FOverriddenPropertyNode;
+enum class EOverriddenPropertyOperation : uint8;
 class UStruct;
 
 namespace PlainProps
@@ -16,16 +18,21 @@ namespace PlainProps
 
 namespace PlainProps::UE
 {
+	struct FMetadataBindings;
+
 
 struct FRestoreContext
 {
 	const FSchemaBindings& Schemas;
 	FCustomBindings& Customs;
+	const FMetadataBindings& Metadatas;
 	TConstArrayView<FStructId> RuntimeIds;
 };
 
 void RestoreStructOverrides(const UStruct* Struct, FStructView StructView, FOverriddenPropertySet& Overrides, const FRestoreContext& Ctx);
 
+void ResetOverriddenPropertyNode(FOverriddenPropertyNode& Node, EOverriddenPropertyOperation Op, TArray<FOverriddenPropertyNode> SubPropertyNodes = TArray<FOverriddenPropertyNode>());
+
 class FScopeRestoreOverrides
 {
 public:
diff --git a/Engine/Plugins/Experimental/PlainPropsUObject/Source/Private/PlainPropsRoundtripTest.cpp b/Engine/Plugins/Experimental/PlainPropsUObject/Source/Private/PlainPropsRoundtripTest.cpp
--- a/Engine/Plugins/Experimental/PlainPropsUObject/Source/Private/PlainPropsRoundtripTest.cpp
+++ b/Engine/Plugins/Experimental/PlainPropsUObject/Source/Private/PlainPropsRoundtripTest.cpp
@@ -604,7 +604,7 @@ public:
 		check(HasObject(ObjectIndex));
 		FStructView StructView = Objects[ObjectMap[ObjectIndex]];
 
-		FRestoreContext Ctx{ GUE.Schemas, GUE.Customs, MakeConstArrayView(LoadStructIds) };
+		FRestoreContext Ctx{ GUE.Schemas, GUE.Customs, GUE.Metadatas, MakeConstArrayView(LoadStructIds) };
 
 		RestoreStructOverrides(Object->GetClass(), StructView, Overrides, Ctx);
 	}
@@ -1717,6 +1717,7 @@ struct FPlainPropsLinkerSandbox
 				{
 					UE_LOGFMT(LogPlainPropsUObject, Log, " Re-binding compiled blueprint class {BindId}", GUE.Debug.Print(Id));
 					GUE.Schemas.DropStruct(Id);
+					GUE.Metadatas.DropMetadata(Id);
 				}
 				else
 				{
@@ -2237,6 +2238,7 @@ static void DropPackages(TArray<UPackage*>&& Packages, FDropPackagesFunc DropPac
 				{
 					UE_LOGFMT(LogPlainPropsUObject, Log, " Dropping loaded dynamic class {BindId}", GUE.Debug.Print(Id));
 					GUE.Schemas.DropStruct(Id);
+					GUE.Metadatas.DropMetadata(Id);
 				}
 			}
 			return true;
diff --git a/Engine/Plugins/Experimental/PlainPropsUObject/Source/Private/PlainPropsSaveOverridable.cpp b/Engine/Plugins/Experimental/PlainPropsUObject/Source/Private/PlainPropsSaveOverridable.cpp
--- a/Engine/Plugins/Experimental/PlainPropsUObject/Source/Private/PlainPropsSaveOverridable.cpp
+++ b/Engine/Plugins/Experimental/PlainPropsUObject/Source/Private/PlainPropsSaveOverridable.cpp
@@ -1,6 +1,7 @@
 // Copyright Epic Games, Inc. All Rights Reserved.
 
 #include "PlainPropsSaveOverridableInternal.h"
+#include "PlainPropsRestoreOverridableInternal.h"
 #include "PlainPropsInternalPrivateMemberPtr.h"
 #include "PlainPropsSave.h"
 #include "PlainPropsSaveMember.h"
@@ -51,7 +52,7 @@ inline const uint8* At(const void* Ptr, SIZE_T Offset)
 	return static_cast<const uint8*>(Ptr) + Offset;
 }
 
-inline bool IsOverriden(const FOverriddenPropertyNode& Override)
+inline bool IsOverridden(const FOverriddenPropertyNode& Override)
 {
 	EOverriddenPropertyOperation Op = Override.GetOperation();
 	return (Op != EOverriddenPropertyOperation::None) && (Op != EOverriddenPropertyOperation::SubObjectsShadowing);
@@ -76,14 +77,14 @@ static void SaveRangeMember(FMemberBuilder& Out, const void* Struct, FMemberId N
 	Out.AddRange(Name, { CreateRangeSchema(Ctx.Scratch, Member), SaveRange(At(Struct, Member.Offset), Member, Ctx) });
 }
 
-static const FStructDeclaration& SaveStructOverrides(FMemberBuilder& Out, const void* Struct, const FOverriddenPropertyNode& Override, FBaseline Base, FBindId Id, const FSaveContext& Ctx);
+static const FStructDeclaration& SaveStructOverrides(FMemberBuilder& Out, const void* Struct, const FOverriddenPropertyNode& Override, FBaseline Base, FBindId Id, const FSaveContext& Ctx, bool bIsRoot);
 
 static void SaveStructMember(FMemberBuilder& Out, const void* Struct, const FOverriddenPropertyNode& Override, 
 							 const void* Default, FMemberId Name, const FSaveContext& Ctx, FStructMemberBinding Member)
 {
 	FMemberBuilder TempBuilder;
 	const void* MemberDefault = Default ? At(Default, Member.Offset) : nullptr;
-	const FStructDeclaration& Declaration = SaveStructOverrides(TempBuilder, At(Struct, Member.Offset), Override, MemberDefault, Member.Id, Ctx);
+	const FStructDeclaration& Declaration = SaveStructOverrides(TempBuilder, At(Struct, Member.Offset), Override, MemberDefault, Member.Id, Ctx, false);
 
 	FBuiltStruct* Result = TempBuilder.IsEmpty() ? nullptr : TempBuilder.BuildAndReset(Ctx.Scratch, Declaration, Ctx.Schemas.GetDebug());
 
@@ -94,9 +95,16 @@ static void SaveStructMember(FMemberBuilder& Out, const void* Struct, const FOve
 }
 
 static void SaveSchemaBoundStructOverrides(FMemberBuilder& Out, const void* Struct, const FOverriddenPropertyNode& Override,
-											FBaseline Base, const FSchemaBinding& Schema, const FStructDeclaration& Declaration, 
-											const FSaveContext& Ctx)
+											FBaseline Base, const FSchemaBinding& Schema, FBindId StructId, const FStructDeclaration& Declaration, 
+											const FSaveContext& Ctx, bool bIsRoot)
 {
+	static FOverriddenPropertyNode DummyReplaceOverride = []() 
+	{ 
+		FOverriddenPropertyNode Node; 
+		ResetOverriddenPropertyNode(Node, EOverriddenPropertyOperation::Replace, {});
+		return Node;
+	}();
+
 	FMemberVisitor It(Schema);
 	if (Declaration.Super)
 	{
@@ -104,18 +112,34 @@ static void SaveSchemaBoundStructOverrides(FMemberBuilder& Out, const void* Stru
 		checkSlow(SuperId == ToOptionalStruct(Declaration.Super));
 		const FStructDeclaration* SuperDecl;
 		const FSchemaBinding& SuperSchema = Ctx.Schemas.GetStruct(SuperId, /* out */ SuperDecl);
-		SaveSchemaBoundStructOverrides(Out, Struct, Override, Base.Super(), SuperSchema, *SuperDecl, Ctx);
+		SaveSchemaBoundStructOverrides(Out, Struct, Override, Base.Super(), SuperSchema, SuperId, *SuperDecl, Ctx, bIsRoot);
 		Out.BuildSuperStruct(Ctx.Scratch, *SuperDecl, Ctx.Schemas.GetDebug());
 	}
 
-	if (IsOverriden(Override))
+	// The override operation for root object doesn't cover always overridden properties, hence the bIsRoot check
+	if (IsOverridden(Override) || bIsRoot)
 	{
 		if (Declaration.Occupancy == EMemberPresence::AllowSparse)
 		{
-			for (FMemberId Name : Declaration.GetMemberOrder())
+			TConstArrayView<EPropertyFlags> MemberFlags = GUE.Metadatas.GetMemberFlags(StructId);
+			TConstArrayView<FMemberId> MemberNames = Declaration.GetMemberOrder();
+			for (int32 Index = 0; Index < MemberNames.Num(); ++Index)
 			{
+				EPropertyFlags Flags = MemberFlags[Index];
+				
+				if ((Flags & CPF_ExperimentalNeverOverriden) != 0)
+				{
+					It.SkipMember();
+					continue;
+				}
+
+				FMemberId Name = MemberNames[Index];
 				const FOverriddenPropertyNode* SubOverride = nullptr;
-				if (Override.GetOperation() == EOverriddenPropertyOperation::Replace)
+				if ((Flags & CPF_ExperimentalAlwaysOverriden) != 0)
+				{
+					SubOverride = &DummyReplaceOverride;
+				}
+				else if (Override.GetOperation() == EOverriddenPropertyOperation::Replace)
 				{
 					// Keep using the same override since a replace operation implies no sub-nodes.
 					SubOverride = &Override;
@@ -123,7 +147,6 @@ static void SaveSchemaBoundStructOverrides(FMemberBuilder& Out, const void* Stru
 				else
 				{
 					FName MemberFName = GUE.Names.ResolveName(Name.Id).ToName();
-
 					SubOverride = FindSubOverride(Override, MemberFName);
 				}
 
@@ -132,7 +155,7 @@ static void SaveSchemaBoundStructOverrides(FMemberBuilder& Out, const void* Stru
 					// The default is only used for custom bindings and specifically overridable container bindings.
 					// When the property operation is 'Replace' we pass no default so that the bindings for overridable
 					// containers do not consider individual items but instead a single 'Assign' operation.
-					const void* Default = (Override.GetOperation() == EOverriddenPropertyOperation::Replace)
+					const void* Default = (SubOverride->GetOperation() == EOverriddenPropertyOperation::Replace)
 											? nullptr : Base.Get();
 
 					switch (It.PeekKind())
@@ -146,7 +169,6 @@ static void SaveSchemaBoundStructOverrides(FMemberBuilder& Out, const void* Stru
 				{
 					It.SkipMember();
 				}
-				
 			}
 			checkSlow(!It.HasMore());
 		}
@@ -158,12 +180,12 @@ static void SaveSchemaBoundStructOverrides(FMemberBuilder& Out, const void* Stru
 }
 
 static const FStructDeclaration& SaveStructOverrides(FMemberBuilder& Out, const void* Struct, const FOverriddenPropertyNode& Override,
-													 FBaseline Base, FBindId Id, const FSaveContext& Ctx)
+													 FBaseline Base, FBindId Id, const FSaveContext& Ctx, bool bIsRoot)
 {
 	const FStructDeclaration* Declaration = nullptr;
 	if (ICustomBinding* Custom = Ctx.Customs.FindStructToSave(Id, /* out */ Declaration))
 	{
-		if (IsOverriden(Override))
+		if (IsOverridden(Override))
 		{
 			Custom->SaveCustom(Out, Struct, Base, Ctx);
 		}
@@ -171,16 +193,16 @@ static const FStructDeclaration& SaveStructOverrides(FMemberBuilder& Out, const
 	else
 	{
 		const FSchemaBinding& Schema = Ctx.Schemas.GetStruct(Id, /* out */ Declaration);
-		SaveSchemaBoundStructOverrides(Out, Struct, Override, Base, Schema, *Declaration, Ctx);
+		SaveSchemaBoundStructOverrides(Out, Struct, Override, Base, Schema, Id, *Declaration, Ctx, bIsRoot);
 	}
 
 	return *Declaration;
 }
 
-static FBuiltStruct* SaveStructOverrides(const void* Struct, const FOverriddenPropertyNode& Override, FBaseline Base, FBindId Id, const FSaveContext& Ctx)
+static FBuiltStruct* SaveStructOverrides(const void* Struct, const FOverriddenPropertyNode& Override, FBaseline Base, FBindId Id, const FSaveContext& Ctx, bool bIsRoot)
 {
 	FMemberBuilder Out;
-	const FStructDeclaration& Declaration = SaveStructOverrides(Out, Struct, Override, Base, Id, Ctx);
+	const FStructDeclaration& Declaration = SaveStructOverrides(Out, Struct, Override, Base, Id, Ctx, bIsRoot);
 	return Out.BuildAndReset(Ctx.Scratch, Declaration, Ctx.Schemas.GetDebug());
 }
 
@@ -192,7 +214,7 @@ FBuiltStruct* SaveStructOverrides(const FOverriddenPropertySet& Overrides,
 	
 	const FOverriddenPropertyNode& RootOverride = Overrides.GetRootOverriddenPropertyNode();
 
-	FBuiltStruct* Result = SaveStructOverrides(Overrides.GetOwner(), RootOverride, Base, BindId, Ctx);
+	FBuiltStruct* Result = SaveStructOverrides(Overrides.GetOwner(), RootOverride, Base, BindId, Ctx, true);
 
 	TlsSaveOverrides = PrevOverrides;
 
diff --git a/Engine/Plugins/Experimental/PlainPropsUObject/Source/Private/PlainPropsUObjectBindings.cpp b/Engine/Plugins/Experimental/PlainPropsUObject/Source/Private/PlainPropsUObjectBindings.cpp
--- a/Engine/Plugins/Experimental/PlainPropsUObject/Source/Private/PlainPropsUObjectBindings.cpp
+++ b/Engine/Plugins/Experimental/PlainPropsUObject/Source/Private/PlainPropsUObjectBindings.cpp
@@ -366,6 +366,7 @@ FCommonMemberIds::FCommonMemberIds(TIdIndexer<FSensitiveName>& Names)
 FGlobals::FGlobals() 
 : Enums(FDebugIds(Names))
 , Schemas(FDebugIds(Names))
+, Metadatas(FDebugIds(Names))
 , Customs(FDebugIds(Names))
 , Scopes(Names)
 , EnumIds(Scopes, Names)
@@ -4083,6 +4084,8 @@ public:
 	void											BindMember(FProperty* Property);
 	void											Interlace(TConstArrayView<FAuxMember> Members);
 	TConstArrayView<FMemberBinding>					GetMembers() const { return Members; }
+	TConstArrayView<EPropertyFlags>					GetMembersFlags() const { return MemberFlags; }
+
 	FStructDeclarationPtr							Declare() const;	
 
 	// Only true for noexport UScriptStruct with STRUCT_Immutable | STRUCT_Atomic flags
@@ -4096,6 +4099,7 @@ private:
 	TArray<FMemberId, TInlineAllocator<64>>			Names;
 	TArray<FMemberBinding, TInlineAllocator<64>>	Members;
 	TArray<FMemberSpec, TInlineAllocator<64>>		Specs;
+	TArray<EPropertyFlags, TInlineAllocator<64>>	MemberFlags;
 
 	// BPVM only?
 	const FName										VerseFunctionProperty{"VerseFunctionProperty"};
@@ -4405,6 +4409,7 @@ void FPropertyBinder::BindMember(FProperty* Property)
 	Members.Emplace(BindProperty(Property));
 	Specs.Emplace(ToSpec(Members.Last()));
 	Names.Emplace(GUE.Names.NameMember(Property->GetFName()));
+	MemberFlags.Emplace(Property->PropertyFlags);
 }
 
 FStructDeclarationPtr FPropertyBinder::Declare() const
@@ -4412,6 +4417,28 @@ FStructDeclarationPtr FPropertyBinder::Declare() const
 	return PlainProps::Declare({LowerCast(Owner), Super, 0, Occupancy, Names, {Specs}});
 }
 
+void FMetadataBindings::BindMetadata(FBindId Id, TConstArrayView<EPropertyFlags> Members)
+{
+	if (Id.Idx >= static_cast<uint32>(Metadatas.Num()))
+	{
+		Metadatas.SetNum(Id.Idx + 1);
+	}
+	checkf(!Metadatas[Id.Idx].IsValid(), TEXT("'%s' metadata already bound"), *Debug.Print(Id));
+	Metadatas[Id.Idx] = FMemberMetadata { .Flags = Members };
+}
+
+void FMetadataBindings::DropMetadata(FBindId Id)
+{
+	checkf(Id.Idx < (uint32)Metadatas.Num(), TEXT("'%s' is unbound"), *Debug.Print(Id));
+	Metadatas[Id.Idx].Reset();
+}
+
+TConstArrayView<EPropertyFlags> FMetadataBindings::GetMemberFlags(FBindId Id) const
+{
+	checkf(Id.Idx < (uint32)Metadatas.Num(), TEXT("'%s' is unbound"), *Debug.Print(Id));
+	return Metadatas[Id.Idx].Flags;
+}
+
 ////////////////////////////////////////////////////////////////////////////////////////////////
 
 // Meta version with similar structure as FPropertyBinder::BindSingleProperty
@@ -4578,6 +4605,7 @@ void BindStruct(FBindId Id, const UStruct* Struct)
 	BindMembers(/* out */ Binder, Struct);
 
 	GUE.Schemas.BindStruct(Id, Binder.GetMembers(), Binder.Declare());
+	GUE.Metadatas.BindMetadata(Id, Binder.GetMembersFlags());
 	GUE.DynamicIds.AddStruct(Struct, Id);
 	
 	// Don't bind CDO defaults, object defaults are passed in from top and objects aren't owned by containers 
diff --git a/Engine/Plugins/Experimental/PlainPropsUObject/Source/Public/PlainPropsUObjectRuntime.h b/Engine/Plugins/Experimental/PlainPropsUObject/Source/Public/PlainPropsUObjectRuntime.h
--- a/Engine/Plugins/Experimental/PlainPropsUObject/Source/Public/PlainPropsUObjectRuntime.h
+++ b/Engine/Plugins/Experimental/PlainPropsUObject/Source/Public/PlainPropsUObjectRuntime.h
@@ -170,6 +170,28 @@ private:
 	TMap<const UStruct*, FDualStructId> Structs;
 };
 
+struct FMemberMetadata
+{
+	bool IsValid() const { return !Flags.IsEmpty(); }
+	void Reset() { Flags.Empty(); }
+
+	TArray<EPropertyFlags> Flags;
+};
+
+struct FMetadataBindings
+{
+	explicit FMetadataBindings(FDebugIds In) : Debug(In) {}
+
+	void 							BindMetadata(FBindId Id, TConstArrayView<EPropertyFlags> Members);
+	void							DropMetadata(FBindId Id);
+
+	TConstArrayView<EPropertyFlags> GetMemberFlags(FBindId Id) const;
+
+private:
+	TArray<FMemberMetadata> Metadatas;
+	FDebugIds				Debug;
+};
+
 struct FGlobals
 {
 	FGlobals();
@@ -177,6 +199,7 @@ struct FGlobals
 	TIdIndexer<FSensitiveName>	Names;
 	FEnumDeclarations			Enums;
 	FSchemaBindings				Schemas;
+	FMetadataBindings 			Metadatas;
 	FCustomBindingsBottom		Customs;
 	FDefaultStructs				Defaults;
 	FCommonScopeIds				Scopes;
