Fix clang & ps5 compilation error

#Synced-CL 52480300

diff --git a/Engine/Plugins/Experimental/PlainPropsUObject/Source/Private/PlainPropsUObjectBindings.cpp b/Engine/Plugins/Experimental/PlainPropsUObject/Source/Private/PlainPropsUObjectBindings.cpp
--- a/Engine/Plugins/Experimental/PlainPropsUObject/Source/Private/PlainPropsUObjectBindings.cpp
+++ b/Engine/Plugins/Experimental/PlainPropsUObject/Source/Private/PlainPropsUObjectBindings.cpp
@@ -4424,7 +4424,7 @@ void FMetadataBindings::BindMetadata(FBindId Id, TConstArrayView<EPropertyFlags>
 		Metadatas.SetNum(Id.Idx + 1);
 	}
 	checkf(!Metadatas[Id.Idx].IsValid(), TEXT("'%s' metadata already bound"), *Debug.Print(Id));
-	Metadatas[Id.Idx] = FMemberMetadata { .Flags = Members };
+	Metadatas[Id.Idx] = FMemberMetadata { .Flags = TArray<EPropertyFlags>(Members) };
 }
 
 void FMetadataBindings::DropMetadata(FBindId Id)
