PlainProps: Add some missing property meta bindings: FArrayProperty/FSetProperty/FMapProperty/FOptionalProperty/FByteProperty

#rnx

Test case with PP_OVERRIDE_OBJECT_LOADING 1
-run=TestPlainProps -allowcommandletaudio -bind=source -batch=linker -load=/ShooterCore/GameplayCues/GCNL_Dash

#Synced-CL 52539327

diff --git a/Engine/Plugins/Experimental/PlainPropsUObject/Source/Internal/PlainPropsUObjectMetaBindingsInternal.h b/Engine/Plugins/Experimental/PlainPropsUObject/Source/Internal/PlainPropsUObjectMetaBindingsInternal.h
--- a/Engine/Plugins/Experimental/PlainPropsUObject/Source/Internal/PlainPropsUObjectMetaBindingsInternal.h
+++ b/Engine/Plugins/Experimental/PlainPropsUObject/Source/Internal/PlainPropsUObjectMetaBindingsInternal.h
@@ -257,6 +257,66 @@ struct FObjectPropertyBaseBinding : ICustomBinding
 	static bool	Diff(const Type& A, const Type& B, const FBindContext& Ctx) { return &A != &B; }
 };
 
+struct FArrayPropertyBinding : ICustomBinding
+{
+	UE_CUSTOM_BIND_MEMBERS(FArrayProperty, InnerPropertyTypeName, InnerProperty);
+	FArrayPropertyBinding(TCustomSpecifier<Ids, NumMembers>& Spec);
+
+	void		Save(FMemberBuilder& Dst, const Type& Src, const Type* Default, const FSaveContext& Ctx) const;
+	void		Load(Type& Dst, FStructLoadView Src, ECustomLoadMethod Method) const;
+	static bool	Diff(const Type& A, const Type& B, const FBindContext& Ctx) { return &A != &B; }
+};
+
+struct FSetPropertyBinding : ICustomBinding
+{
+	UE_CUSTOM_BIND_MEMBERS(FSetProperty, InnerPropertyTypeName, InnerProperty);
+	FSetPropertyBinding(TCustomSpecifier<Ids, NumMembers>& Spec);
+
+	void		Save(FMemberBuilder& Dst, const Type& Src, const Type* Default, const FSaveContext& Ctx) const;
+	void		Load(Type& Dst, FStructLoadView Src, ECustomLoadMethod Method) const;
+	static bool	Diff(const Type& A, const Type& B, const FBindContext& Ctx) { return &A != &B; }
+};
+
+struct FMapPropertyBinding : ICustomBinding
+{
+	UE_CUSTOM_BIND_MEMBERS(FMapProperty, KeyPropertyTypeName, KeyProperty, ValuePropertyTypeName, ValueProperty);
+	FMapPropertyBinding(TCustomSpecifier<Ids, NumMembers>& Spec);
+
+	void		Save(FMemberBuilder& Dst, const Type& Src, const Type* Default, const FSaveContext& Ctx) const;
+	void		Load(Type& Dst, FStructLoadView Src, ECustomLoadMethod Method) const;
+	static bool	Diff(const Type& A, const Type& B, const FBindContext& Ctx) { return &A != &B; }
+};
+
+struct FOptionalPropertyBinding : ICustomBinding
+{
+	UE_CUSTOM_BIND_MEMBERS(FOptionalProperty, ValuePropertyTypeName, ValueProperty);
+	FOptionalPropertyBinding(TCustomSpecifier<Ids, NumMembers>& Spec);
+
+	void		Save(FMemberBuilder& Dst, const Type& Src, const Type* Default, const FSaveContext& Ctx) const;
+	void		Load(Type& Dst, FStructLoadView Src, ECustomLoadMethod Method) const;
+	static bool	Diff(const Type& A, const Type& B, const FBindContext& Ctx) { return &A != &B; }
+};
+
+struct FBytePropertyBinding : ICustomBinding
+{
+	UE_CUSTOM_BIND_MEMBERS(FByteProperty, Enum);
+	FBytePropertyBinding(TCustomSpecifier<Ids, NumMembers>& Spec);
+
+	void		Save(FMemberBuilder& Dst, const Type& Src, const Type* Default, const FSaveContext& Ctx) const;
+	void		Load(Type& Dst, FStructLoadView Src, ECustomLoadMethod Method) const;
+	static bool	Diff(const Type& A, const Type& B, const FBindContext& Ctx) { return &A != &B; }
+};
+
+struct FEnumPropertyBinding : ICustomBinding
+{
+	UE_CUSTOM_BIND_MEMBERS(FEnumProperty, UnderlyingPropertyTypeName, UnderlyingProperty, Enum);
+	FEnumPropertyBinding(TCustomSpecifier<Ids, NumMembers>& Spec);
+
+	void		Save(FMemberBuilder& Dst, const Type& Src, const Type* Default, const FSaveContext& Ctx) const;
+	void		Load(Type& Dst, FStructLoadView Src, ECustomLoadMethod Method) const;
+	static bool	Diff(const Type& A, const Type& B, const FBindContext& Ctx) { return &A != &B; }
+};
+
 ////////////////////////////////////////////////////////////////////////////////////////////////
 // Native UField meta bindings
 
@@ -327,10 +387,16 @@ namespace PlainProps
 {
 // custom FProperty meta bindings
 template<> struct TCustomBind<FProperty> { using Type = UE::FPropertyBinding; };
+template<> struct TCustomBind<FEnumProperty> { using Type = UE::FEnumPropertyBinding; };
 template<> struct TCustomBind<FBoolProperty> { using Type = UE::FBoolPropertyBinding; };
+template<> struct TCustomBind<FByteProperty> { using Type = UE::FBytePropertyBinding; };
 template<> struct TCustomBind<FStructProperty> { using Type = UE::FStructPropertyBinding; };
 template<> struct TCustomBind<FClassProperty> { using Type = UE::FClassPropertyBinding; };
 template<> struct TCustomBind<FObjectPropertyBase> { using Type = UE::FObjectPropertyBaseBinding; };
+template<> struct TCustomBind<FArrayProperty> { using Type = UE::FArrayPropertyBinding; };
+template<> struct TCustomBind<FSetProperty> { using Type = UE::FSetPropertyBinding; };
+template<> struct TCustomBind<FMapProperty> { using Type = UE::FMapPropertyBinding; };
+template<> struct TCustomBind<FOptionalProperty> { using Type = UE::FOptionalPropertyBinding; };
 
 // custom UField meta bindings
 template<> struct TCustomBind<UStruct> { using Type = UE::UStructBinding; };
diff --git a/Engine/Plugins/Experimental/PlainPropsUObject/Source/Private/PlainPropsUObjectBindings.cpp b/Engine/Plugins/Experimental/PlainPropsUObject/Source/Private/PlainPropsUObjectBindings.cpp
--- a/Engine/Plugins/Experimental/PlainPropsUObject/Source/Private/PlainPropsUObjectBindings.cpp
+++ b/Engine/Plugins/Experimental/PlainPropsUObject/Source/Private/PlainPropsUObjectBindings.cpp
@@ -310,6 +310,7 @@ FCommonPropertyStructIds::FCommonPropertyStructIds(const FCommonScopeIds& Scopes
 : Property(					Names.IndexStruct({Scopes.CoreUObject, Names.MakeTypename("FProperty")}))
 , EnumProperty(				Names.IndexStruct({Scopes.CoreUObject, Names.MakeTypename("FEnumProperty")}))
 , BoolProperty(				Names.IndexStruct({Scopes.CoreUObject, Names.MakeTypename("FBoolProperty")}))
+, ByteProperty(				Names.IndexStruct({Scopes.CoreUObject, Names.MakeTypename("FByteProperty") }))
 , StructProperty(			Names.IndexStruct({Scopes.CoreUObject, Names.MakeTypename("FStructProperty")}))
 , ClassProperty(			Names.IndexStruct({Scopes.CoreUObject, Names.MakeTypename("FClassProperty")}))
 , ObjectPropertyBase(		Names.IndexStruct({Scopes.CoreUObject, Names.MakeTypename("FObjectPropertyBase")}))
@@ -4450,6 +4451,11 @@ FDualStructId IndexPropertyMeta(const FField* Property)
 	{
 		if (HasAny<CASTCLASS_FNumericProperty>(Flags))
 		{
+			if (HasAny<CASTCLASS_FByteProperty>(Flags))
+			{
+				return GUE.Properties.ByteProperty;
+			}
+
 			return GUE.Properties.Property;
 		}
 		return HasAny<CASTCLASS_FEnumProperty>(Flags)
@@ -4638,9 +4644,16 @@ static TArray<FBindId> GetStructsToSkip(EBatchType BatchType)
 static void InitPropertiesMeta()
 {
 	GUE.Defaults.BindZeroes(GUE.Properties.Property, sizeof(FProperty), alignof(FProperty));
+	GUE.Defaults.BindZeroes(GUE.Properties.EnumProperty, sizeof(FEnumProperty), alignof(FEnumProperty));
 	GUE.Defaults.BindZeroes(GUE.Properties.BoolProperty, sizeof(FBoolProperty), alignof(FBoolProperty));
+	GUE.Defaults.BindZeroes(GUE.Properties.ByteProperty, sizeof(FByteProperty), alignof(FByteProperty));
 	GUE.Defaults.BindZeroes(GUE.Properties.StructProperty, sizeof(FStructProperty), alignof(FStructProperty));
 	GUE.Defaults.BindZeroes(GUE.Properties.ClassProperty, sizeof(FClassProperty), alignof(FClassProperty));
+	GUE.Defaults.BindZeroes(GUE.Properties.ObjectPropertyBase, sizeof(FObjectPropertyBase), alignof(FObjectPropertyBase));
+	GUE.Defaults.BindZeroes(GUE.Properties.ArrayProperty, sizeof(FArrayProperty), alignof(FArrayProperty));
+	GUE.Defaults.BindZeroes(GUE.Properties.SetProperty, sizeof(FSetProperty), alignof(FSetProperty));
+	GUE.Defaults.BindZeroes(GUE.Properties.MapProperty, sizeof(FMapProperty), alignof(FMapProperty));
+	GUE.Defaults.BindZeroes(GUE.Properties.OptionalProperty, sizeof(FOptionalProperty), alignof(FOptionalProperty));
 }
 
 static void InitBatchedProperties(EBatchType BatchType)
@@ -4734,9 +4747,15 @@ struct FPlainPropsBindingScope
 		, StructFlags()
 		, PropertyFlags()
 		, Property(GUE.Properties.Property)
+		, EnumProperty(GUE.Properties.EnumProperty)
 		, BoolProperty(GUE.Properties.BoolProperty)
+		, ByteProperty(GUE.Properties.ByteProperty)
 		, StructProperty(GUE.Properties.StructProperty)
 		, ClassProperty(GUE.Properties.ClassProperty)
+		, ArrayProperty(GUE.Properties.ArrayProperty)
+		, SetProperty(GUE.Properties.SetProperty)
+		, MapProperty(GUE.Properties.MapProperty)
+		, OptionalProperty(GUE.Properties.OptionalProperty)
 		, ObjectPropertyBase(GUE.Properties.ObjectPropertyBase)
 		, Struct(GUE.Structs.Struct)
 		, Class(GUE.Structs.Class)
@@ -4781,9 +4800,15 @@ struct FPlainPropsBindingScope
 	TScopedEnumDeclaration<EStructFlags, EEnumMode::Flag, FDefaultRuntime> StructFlags;
 	TScopedEnumDeclaration<EPropertyFlags, EEnumMode::Flag, FDefaultRuntime> PropertyFlags;
 	TScopedStructBinding<FProperty, FDefaultRuntime> Property;
+	TScopedStructBinding<FEnumProperty, FDefaultRuntime> EnumProperty;
 	TScopedStructBinding<FBoolProperty, FDefaultRuntime> BoolProperty;
+	TScopedStructBinding<FByteProperty, FDefaultRuntime> ByteProperty;
 	TScopedStructBinding<FStructProperty, FDefaultRuntime> StructProperty;
 	TScopedStructBinding<FClassProperty, FDefaultRuntime> ClassProperty;
+	TScopedStructBinding<FArrayProperty, FDefaultRuntime> ArrayProperty;
+	TScopedStructBinding<FSetProperty, FDefaultRuntime> SetProperty;
+	TScopedStructBinding<FMapProperty, FDefaultRuntime> MapProperty;
+	TScopedStructBinding<FOptionalProperty, FDefaultRuntime> OptionalProperty;
 	TScopedStructBinding<FObjectPropertyBase, FDefaultRuntime> ObjectPropertyBase;
 	TScopedStructBinding<UStruct, FDefaultRuntime> Struct;
 	TScopedStructBinding<UClass, FDefaultRuntime> Class;
diff --git a/Engine/Plugins/Experimental/PlainPropsUObject/Source/Private/PlainPropsUObjectMetaBindings.cpp b/Engine/Plugins/Experimental/PlainPropsUObject/Source/Private/PlainPropsUObjectMetaBindings.cpp
--- a/Engine/Plugins/Experimental/PlainPropsUObject/Source/Private/PlainPropsUObjectMetaBindings.cpp
+++ b/Engine/Plugins/Experimental/PlainPropsUObject/Source/Private/PlainPropsUObjectMetaBindings.cpp
@@ -39,6 +39,12 @@ PP_DEFINE_PRIVATE_MEMBER_PTR(FBoolProperty, ByteOffset,	uint8);
 PP_DEFINE_PRIVATE_MEMBER_PTR(FBoolProperty, ByteMask,	uint8);
 PP_DEFINE_PRIVATE_MEMBER_PTR(FBoolProperty, FieldMask,	uint8);
 PP_DEFINE_PRIVATE_MEMBER_NAMESPACE_END()
+
+PP_DEFINE_PRIVATE_MEMBER_NAMESPACE_BEGIN(FEnumProperty_Private)
+PP_DEFINE_PRIVATE_MEMBER_PTR(FEnumProperty, UnderlyingProp, FNumericProperty*);
+PP_DEFINE_PRIVATE_MEMBER_PTR(FEnumProperty, Enum, TObjectPtr<UEnum>);
+PP_DEFINE_PRIVATE_MEMBER_NAMESPACE_END()
+
 PRAGMA_ENABLE_DEPRECATION_WARNINGS
 
 namespace PlainProps::UE
@@ -47,6 +53,31 @@ namespace PlainProps::UE
 ////////////////////////////////////////////////////////////////////////////////////////////////
 // Native FProperty meta bindings
 
+static void SaveSingleProperty(FMemberBuilder& Dst, const FMemberId& TypeNameMemberId, const FMemberId& PropertyMemberId, FProperty* InProperty, const FSaveContext& Ctx)
+{
+	const FName PropertyTypeName = InProperty ? InProperty->GetClass()->GetFName() : FName();
+	Dst.AddStruct(TypeNameMemberId, GUE.Structs.Name, SaveStruct(&PropertyTypeName, GUE.Structs.Name, Ctx));
+
+	if (InProperty)
+	{
+		FBindId Id = IndexPropertyMeta(InProperty);
+		Dst.AddStruct(PropertyMemberId, Id, SaveStruct(InProperty, Id, Ctx));
+	}
+}
+
+static void LoadSingleProperty(FMemberLoader& Members, FProperty*& Dst, const FProperty& Owner)
+{
+	FName PropertyTypeName;
+	LoadStruct(&PropertyTypeName, Members.GrabStruct());
+
+	if (!PropertyTypeName.IsNone())
+	{
+		Dst = CastField<FProperty>(FField::Construct(PropertyTypeName, &Owner, FName()));
+		FBindId Id = IndexPropertyMeta(Dst);
+		LoadStruct(Dst, Members.GrabStruct());
+	}
+}
+
 FPropertyBinding::FPropertyBinding(TCustomSpecifier<Ids, NumMembers>& Spec)
 	: PropertyFlagsId(GetEnumId<Ids, EPropertyFlags>())
 	, LifetimeConditionId(Ids::IndexEnum(GUE.Scopes.CoreUObject, "ELifetimeCondition"))
@@ -106,7 +137,7 @@ FBoolPropertyBinding::FBoolPropertyBinding(TCustomSpecifier<Ids, NumMembers>& Sp
 void FBoolPropertyBinding::Save(FMemberBuilder& Dst, const FBoolProperty& Src, const FBoolProperty*, const FSaveContext& Ctx) const
 {
 	using namespace FBoolProperty_Private;
-	SaveSuper(Dst, &Src, FBindId(GUE.Properties.Property), Ctx);
+	SaveSuper(Dst, &Src, GUE.Properties.Property, Ctx);
 	Dst.Add(MemberIds[M::FieldSize], Src.*_FieldSize);
 	Dst.Add(MemberIds[M::ByteOffset], Src.*_ByteOffset);
 	Dst.Add(MemberIds[M::ByteMask], Src.*_ByteMask);
@@ -133,7 +164,7 @@ FStructPropertyBinding::FStructPropertyBinding(TCustomSpecifier<Ids, NumMembers>
 
 void FStructPropertyBinding::Save(FMemberBuilder& Dst, const FStructProperty& Src, const FStructProperty*, const FSaveContext& Ctx) const
 {
-	SaveSuper(Dst, &Src, FBindId(GUE.Properties.Property), Ctx);
+	SaveSuper(Dst, &Src, GUE.Properties.Property, Ctx);
 	Dst.AddStruct(MemberIds[M::Struct], GUE.Structs.ObjectPtr, SaveStruct(&Src.Struct, GUE.Structs.ObjectPtr, Ctx));
 }
 
@@ -153,7 +184,7 @@ FClassPropertyBinding::FClassPropertyBinding(TCustomSpecifier<Ids, NumMembers>&
 
 void FClassPropertyBinding::Save(FMemberBuilder& Dst, const FClassProperty& Src, const FClassProperty*, const FSaveContext& Ctx) const
 {
-	SaveSuper(Dst, &Src, FBindId(GUE.Properties.Property), Ctx);
+	SaveSuper(Dst, &Src, GUE.Properties.Property, Ctx);
 	Dst.AddStruct(MemberIds[M::MetaClass], GUE.Structs.ObjectPtr, SaveStruct(&Src.MetaClass, GUE.Structs.ObjectPtr, Ctx));
 }
 
@@ -180,7 +211,7 @@ FObjectPropertyBaseBinding::FObjectPropertyBaseBinding(TCustomSpecifier<Ids, Num
 
 void FObjectPropertyBaseBinding::Save(FMemberBuilder& Dst, const FObjectPropertyBase& Src, const FObjectPropertyBase*, const FSaveContext& Ctx) const
 {
-	SaveSuper(Dst, &Src, FBindId(GUE.Properties.Property), Ctx);
+	SaveSuper(Dst, &Src, GUE.Properties.Property, Ctx);
 	Dst.AddStruct(MemberIds[M::PropertyClass], GUE.Structs.ObjectPtr, SaveStruct(&Src.PropertyClass, GUE.Structs.ObjectPtr, Ctx));
 }
 
@@ -192,6 +223,160 @@ void FObjectPropertyBaseBinding::Load(FObjectPropertyBase& Dst, FStructLoadView
 	check(!Members.HasMore());
 }
 
+FArrayPropertyBinding::FArrayPropertyBinding(TCustomSpecifier<Ids, NumMembers>& Spec)
+{
+	Spec.Super = FDeclId(GUE.Properties.Property);
+	// Intentionally skip ArrayFlags (no support for blueprint UsesMemoryImageAllocator)
+	Spec.Members[M::InnerPropertyTypeName] = FDeclId(GUE.Structs.Name);
+	Spec.Members[M::InnerProperty] = SpecDynamicStruct;
+}
+
+void FArrayPropertyBinding::Save(FMemberBuilder& Dst, const FArrayProperty& Src, const FArrayProperty*, const FSaveContext& Ctx) const
+{
+	SaveSuper(Dst, &Src, GUE.Properties.Property, Ctx);
+	
+	SaveSingleProperty(Dst, MemberIds[M::InnerPropertyTypeName], MemberIds[M::InnerProperty], Src.Inner, Ctx);
+}
+
+void FArrayPropertyBinding::Load(FArrayProperty& Dst, FStructLoadView Src, ECustomLoadMethod Method) const
+{
+	FMemberLoader Members(Src);
+	TryLoadSuper(&Dst, Members);
+	
+	LoadSingleProperty(Members, Dst.Inner, Dst);
+
+	check(!Members.HasMore());
+}
+
+FSetPropertyBinding::FSetPropertyBinding(TCustomSpecifier<Ids, NumMembers>& Spec)
+{
+	Spec.Super = FDeclId(GUE.Properties.Property);
+	Spec.Members[M::InnerPropertyTypeName] = FDeclId(GUE.Structs.Name);
+	Spec.Members[M::InnerProperty] = SpecDynamicStruct;
+}
+
+void FSetPropertyBinding::Save(FMemberBuilder& Dst, const FSetProperty& Src, const FSetProperty*, const FSaveContext& Ctx) const
+{
+	SaveSuper(Dst, &Src, GUE.Properties.Property, Ctx);
+
+	SaveSingleProperty(Dst, MemberIds[M::InnerPropertyTypeName], MemberIds[M::InnerProperty], Src.ElementProp, Ctx);
+}
+
+void FSetPropertyBinding::Load(FSetProperty& Dst, FStructLoadView Src, ECustomLoadMethod Method) const
+{
+	FMemberLoader Members(Src);
+	TryLoadSuper(&Dst, Members);
+
+	LoadSingleProperty(Members, Dst.ElementProp, Dst);
+
+	check(!Members.HasMore());
+}
+
+FMapPropertyBinding::FMapPropertyBinding(TCustomSpecifier<Ids, NumMembers>& Spec)
+{
+	Spec.Super = FDeclId(GUE.Properties.Property);
+	Spec.Members[M::KeyPropertyTypeName] = FDeclId(GUE.Structs.Name);
+	Spec.Members[M::KeyProperty] = SpecDynamicStruct;
+	Spec.Members[M::ValuePropertyTypeName] = FDeclId(GUE.Structs.Name);
+	Spec.Members[M::ValueProperty] = SpecDynamicStruct;
+}
+
+void FMapPropertyBinding::Save(FMemberBuilder& Dst, const FMapProperty& Src, const FMapProperty*, const FSaveContext& Ctx) const
+{
+	SaveSuper(Dst, &Src, GUE.Properties.Property, Ctx);
+
+	SaveSingleProperty(Dst, MemberIds[M::KeyPropertyTypeName], MemberIds[M::KeyProperty], Src.KeyProp, Ctx);
+	SaveSingleProperty(Dst, MemberIds[M::ValuePropertyTypeName], MemberIds[M::ValueProperty], Src.ValueProp, Ctx);
+}
+
+void FMapPropertyBinding::Load(FMapProperty& Dst, FStructLoadView Src, ECustomLoadMethod Method) const
+{
+	FMemberLoader Members(Src);
+	TryLoadSuper(&Dst, Members);
+
+	LoadSingleProperty(Members, Dst.KeyProp, Dst);
+	LoadSingleProperty(Members, Dst.ValueProp, Dst);
+
+	check(!Members.HasMore());
+}
+
+FOptionalPropertyBinding::FOptionalPropertyBinding(TCustomSpecifier<Ids, NumMembers>& Spec)
+{
+	Spec.Super = FDeclId(GUE.Properties.Property);
+	Spec.Members[M::ValuePropertyTypeName] = FDeclId(GUE.Structs.Name);
+	Spec.Members[M::ValueProperty] = SpecDynamicStruct;
+}
+
+void FOptionalPropertyBinding::Save(FMemberBuilder& Dst, const FOptionalProperty& Src, const FOptionalProperty* Default, const FSaveContext& Ctx) const
+{
+	SaveSuper(Dst, &Src, GUE.Properties.Property, Ctx);
+	SaveSingleProperty(Dst, MemberIds[M::ValuePropertyTypeName], MemberIds[M::ValueProperty], Src.GetValueProperty(), Ctx);
+}
+
+void FOptionalPropertyBinding::Load(FOptionalProperty& Dst, FStructLoadView Src, ECustomLoadMethod Method) const
+{
+	FMemberLoader Members(Src);
+	TryLoadSuper(&Dst, Members);
+
+	FProperty* ValueProperty = nullptr;
+	LoadSingleProperty(Members, ValueProperty, Dst);
+	Dst.SetValueProperty(ValueProperty);
+
+	check(!Members.HasMore());
+}
+
+FBytePropertyBinding::FBytePropertyBinding(TCustomSpecifier<Ids, NumMembers>& Spec)
+{
+	Spec.Super = FDeclId(GUE.Properties.Property);
+	Spec.Members[M::Enum] = FDeclId(GUE.Structs.ObjectPtr);
+}
+
+void FBytePropertyBinding::Save(FMemberBuilder& Dst, const FByteProperty& Src, const FByteProperty* Default, const FSaveContext& Ctx) const
+{
+	SaveSuper(Dst, &Src, GUE.Properties.Property, Ctx);
+	Dst.AddStruct(MemberIds[M::Enum], GUE.Structs.ObjectPtr, SaveStruct(&Src.Enum, GUE.Structs.ObjectPtr, Ctx));
+}
+
+void FBytePropertyBinding::Load(FByteProperty& Dst, FStructLoadView Src, ECustomLoadMethod Method) const
+{
+	FMemberLoader Members(Src);
+	TryLoadSuper(&Dst, Members);
+	LoadStruct(&Dst.Enum, Members.GrabStruct());
+	check(!Members.HasMore());
+}
+
+FEnumPropertyBinding::FEnumPropertyBinding(TCustomSpecifier<Ids, NumMembers>& Spec)
+{
+	Spec.Super = FDeclId(GUE.Properties.Property);
+	Spec.Members[M::UnderlyingPropertyTypeName] = FDeclId(GUE.Structs.Name);
+	Spec.Members[M::UnderlyingProperty] = SpecDynamicStruct;
+	Spec.Members[M::Enum] = FDeclId(GUE.Structs.ObjectPtr);
+}
+
+void FEnumPropertyBinding::Save(FMemberBuilder& Dst, const FEnumProperty& Src, const FEnumProperty* Default, const FSaveContext& Ctx) const
+{
+	using namespace FEnumProperty_Private;
+	SaveSuper(Dst, &Src, GUE.Properties.Property, Ctx);
+
+	SaveSingleProperty(Dst, MemberIds[M::UnderlyingPropertyTypeName], MemberIds[M::UnderlyingProperty], Src.*_UnderlyingProp, Ctx);
+
+	Dst.AddStruct(MemberIds[M::Enum], GUE.Structs.ObjectPtr, SaveStruct(&(Src.*_Enum), GUE.Structs.ObjectPtr, Ctx));
+}
+
+void FEnumPropertyBinding::Load(FEnumProperty& Dst, FStructLoadView Src, ECustomLoadMethod Method) const
+{
+	using namespace FEnumProperty_Private;
+	FMemberLoader Members(Src);
+	TryLoadSuper(&Dst, Members);
+	
+	FProperty* UnderlyingProp = nullptr;
+	LoadSingleProperty(Members, UnderlyingProp, Dst);
+	Dst.*_UnderlyingProp = CastField<FNumericProperty>(UnderlyingProp);
+
+	LoadStruct(&(Dst.*_Enum), Members.GrabStruct());
+	check(!Members.HasMore());
+}
+
 ////////////////////////////////////////////////////////////////////////////////////////////////
 
 UStructBinding::UStructBinding(TCustomSpecifier<Ids, NumMembers>& Spec)
@@ -366,7 +551,7 @@ void UClassBinding::LoadInterfacesStatic(TArray<FImplementedInterface>& Interfac
 
 void UClassBinding::Save(FMemberBuilder& Dst, const UClass& Src, const UClass*, const FSaveContext& Ctx) const
 {
-	SaveSuper(Dst, &Src, FBindId(GUE.Structs.Struct), Ctx);
+	SaveSuper(Dst, &Src, GUE.Structs.Struct, Ctx);
 	EClassFlags SavedClassFlags = Src.ClassFlags & ~CLASS_ShouldNeverBeLoaded;
 	Dst.AddEnum(MemberIds[M::ClassFlags], ClassFlagsId, SavedClassFlags);
 	Dst.AddStruct(MemberIds[M::ClassWithin], GUE.Structs.ObjectPtr, SaveStruct(&Src.ClassWithin, GUE.Structs.ObjectPtr, Ctx));
@@ -446,7 +631,7 @@ UFunctionBinding::UFunctionBinding(TCustomSpecifier<Ids, NumMembers>& Spec)
 
 void UFunctionBinding::Save(FMemberBuilder& Dst, const UFunction& Src, const UFunction*, const FSaveContext& Ctx) const
 {
-	SaveSuper(Dst, &Src, FBindId(GUE.Structs.Struct), Ctx);
+	SaveSuper(Dst, &Src, GUE.Structs.Struct, Ctx);
 	Dst.AddEnum(MemberIds[M::FunctionFlags], FunctionFlagsId, Src.FunctionFlags);
 	Dst.AddStruct(MemberIds[M::EventGraphFunction], GUE.Structs.ObjectPtr, SaveStruct(&Src.EventGraphFunction, GUE.Structs.ObjectPtr, Ctx));
 	Dst.Add(MemberIds[M::EventGraphCallOffset], Src.EventGraphCallOffset);
@@ -472,7 +657,7 @@ UScriptStructBinding::UScriptStructBinding(TCustomSpecifier<Ids, NumMembers>& Sp
 
 void UScriptStructBinding::Save(FMemberBuilder& Dst, const UScriptStruct& Src, const UScriptStruct*, const FSaveContext& Ctx) const
 {
-	SaveSuper(Dst, &Src, FBindId(GUE.Structs.Struct), Ctx);
+	SaveSuper(Dst, &Src, GUE.Structs.Struct, Ctx);
 	EStructFlags SavedStructFlags = EStructFlags(Src.StructFlags & ~STRUCT_ComputedFlags);
 	Dst.AddEnum(MemberIds[M::StructFlags], StructFlagsId, SavedStructFlags);
 }
diff --git a/Engine/Plugins/Experimental/PlainPropsUObject/Source/Public/PlainPropsUObjectRuntime.h b/Engine/Plugins/Experimental/PlainPropsUObject/Source/Public/PlainPropsUObjectRuntime.h
--- a/Engine/Plugins/Experimental/PlainPropsUObject/Source/Public/PlainPropsUObjectRuntime.h
+++ b/Engine/Plugins/Experimental/PlainPropsUObject/Source/Public/PlainPropsUObjectRuntime.h
@@ -90,6 +90,7 @@ struct FCommonPropertyStructIds
 	FDualStructId				Property;
 	FDualStructId				EnumProperty;
 	FDualStructId				BoolProperty;
+	FDualStructId				ByteProperty;
 	FDualStructId				StructProperty;
 	FDualStructId				ClassProperty;
 	FDualStructId				ObjectPropertyBase;
