Fix game non-unity build

#rnx

#Synced-CL 52572927

diff --git a/Engine/Plugins/Experimental/PlainPropsUObject/Source/Private/PlainPropsUObjectMetaBindings.cpp b/Engine/Plugins/Experimental/PlainPropsUObject/Source/Private/PlainPropsUObjectMetaBindings.cpp
--- a/Engine/Plugins/Experimental/PlainPropsUObject/Source/Private/PlainPropsUObjectMetaBindings.cpp
+++ b/Engine/Plugins/Experimental/PlainPropsUObject/Source/Private/PlainPropsUObjectMetaBindings.cpp
@@ -11,6 +11,8 @@
 #include "UObject/ObjectPtr.h"
 #include "UObject/Package.h"
 #include "UObject/UnrealType.h"
+#include "UObject/EnumProperty.h"
+#include "UObject/PropertyOptional.h"
 
 // Temp hacks. Long-term either add FProperty getters for ctor/dtor/hash function pointers 
 // and delegate APIs for non-intrusive serialization or integrate PlainProps into Core/CoreUObject
