PlainProps: Ignore override diff when semantically identical:
modified node with all overridden sub property is semantically identical to a replace node

#Synced-CL 52521357

diff --git a/Engine/Plugins/Experimental/PlainPropsUObject/Source/Private/PlainPropsSaveOverridable.cpp b/Engine/Plugins/Experimental/PlainPropsUObject/Source/Private/PlainPropsSaveOverridable.cpp
--- a/Engine/Plugins/Experimental/PlainPropsUObject/Source/Private/PlainPropsSaveOverridable.cpp
+++ b/Engine/Plugins/Experimental/PlainPropsUObject/Source/Private/PlainPropsSaveOverridable.cpp
@@ -23,25 +23,29 @@ namespace PlainProps::UE
 // Todo: This should be moved to an appropriate context that we can pass around.
 thread_local const FOverriddenPropertySet* TlsSaveOverrides = nullptr;
 
+FName GetOverridePropertyName(const FOverriddenPropertyNode& Override)
+{
+	FName Name;
+	const FOverriddenPropertyPath* Path = nullptr;
+	{
+		using namespace FOverriddenPropertyNodeID_Private;
+		Path = &(Override.GetNodeID().*_Path);
+	}
+	{
+		using namespace FOverriddenPropertyPath_Private;
+		Name = Path->IsEmpty() ? FName() : ((*Path).*_Path)[0];
+	}
+	return Name;
+}
+
 inline const FOverriddenPropertyNode* FindSubOverride(const FOverriddenPropertyNode& Override, FName MemberName)
 {
 	for (const FOverriddenPropertyNode& SubNode : Override.GetSubPropertyNodes())
 	{
-		const FOverriddenPropertyPath* Path = nullptr;
-
-		{
-			using namespace FOverriddenPropertyNodeID_Private;
-			Path = &(SubNode.GetNodeID().*_Path);
-		}
-		
+		FName IdName = GetOverridePropertyName(SubNode);
+		if (IdName == MemberName)
 		{
-			using namespace FOverriddenPropertyPath_Private;
-			FName IdName = Path->IsEmpty() ? FName() : ((*Path).*_Path)[0];
-
-			if (IdName == MemberName)
-			{
-				return &SubNode;
-			}
+			return &SubNode;
 		}
 	}
 	return nullptr;
diff --git a/Engine/Plugins/Experimental/PlainPropsUObject/Source/Private/PlainPropsSaveOverridableInternal.h b/Engine/Plugins/Experimental/PlainPropsUObject/Source/Private/PlainPropsSaveOverridableInternal.h
--- a/Engine/Plugins/Experimental/PlainPropsUObject/Source/Private/PlainPropsSaveOverridableInternal.h
+++ b/Engine/Plugins/Experimental/PlainPropsUObject/Source/Private/PlainPropsSaveOverridableInternal.h
@@ -6,6 +6,7 @@
 #include "PlainPropsSave.h"
 
 struct FOverriddenPropertySet;
+struct FOverriddenPropertyNode;
 
 namespace PlainProps
 {
@@ -13,6 +14,8 @@ namespace PlainProps
 namespace UE
 {
 
+FName GetOverridePropertyName(const FOverriddenPropertyNode& Override);
+
 FBuiltStruct* SaveStructOverrides(const FOverriddenPropertySet& Overrides, FBaseline Base, FBindId BindId, const FSaveContext& Context);
 
 const FOverriddenPropertySet* GetSaveOverrides();
diff --git a/Engine/Plugins/Experimental/PlainPropsUObject/Source/Private/PlainPropsUObjectDiff.cpp b/Engine/Plugins/Experimental/PlainPropsUObject/Source/Private/PlainPropsUObjectDiff.cpp
--- a/Engine/Plugins/Experimental/PlainPropsUObject/Source/Private/PlainPropsUObjectDiff.cpp
+++ b/Engine/Plugins/Experimental/PlainPropsUObject/Source/Private/PlainPropsUObjectDiff.cpp
@@ -2,6 +2,7 @@
 
 #include "PlainPropsUObjectDiffInternal.h"
 #include "PlainPropsInternalPrivateMemberPtr.h"
+#include "PlainPropsSaveOverridableInternal.h"
 #include "UObject/Class.h"
 #include "UObject/FieldPathProperty.h"
 #include "UObject/Object.h"
@@ -654,45 +655,107 @@ static bool EqualOverrideIds(const FOverriddenPropertyNodeID& A, const FOverridd
 	}
 }
 
-static bool DiffOverrideNodes(const FOverriddenPropertyNode& A, const FOverriddenPropertyNode& B)
+struct FDiffNodeContext
+{
+	FDiffNodeContext(const UObject* A)
+	: ObjectA(A)
+	, CurrentStruct(A->GetClass())
+	{}
+
+	const UObject* ObjectA = nullptr;
+	UStruct* CurrentStruct = nullptr;
+	FPropertyVisitorPath VisitorPath;
+};
+
+FDiffNodeContext MakeNewContext(const FOverriddenPropertyNode& SubNodeA, const FDiffNodeContext& Ctx)
+{
+	const UObject* SubA = Ctx.ObjectA;
+	{
+		using namespace FOverriddenPropertyNodeIdDiff_Private;
+		SubA = (SubNodeA.GetNodeID().*_Object).Get();
+	}
+	FDiffNodeContext NewCtx(SubA ? SubA : Ctx.ObjectA);
+	NewCtx.VisitorPath = Ctx.VisitorPath;
+	FName PropertyName = GetOverridePropertyName(SubNodeA);
+	FProperty* Prop = Ctx.CurrentStruct->FindPropertyByName(PropertyName);
+	if (Prop)
+	{
+		FPropertyVisitorInfo PropInfo(Prop, Ctx.CurrentStruct);
+		NewCtx.VisitorPath.Push(PropInfo);
+		if (FStructProperty* StructProp = CastField<FStructProperty>(Prop))
+		{
+			NewCtx.CurrentStruct = StructProp->Struct;
+		}
+	}
+	return NewCtx;
+}
+
+static bool DiffOverrideNodes(const FOverriddenPropertyNode& A, const FOverriddenPropertyNode& B, FDiffNodeContext& Ctx)
 {
 	EOverriddenPropertyOperation OpA = A.GetOperation();
 	EOverriddenPropertyOperation OpB = B.GetOperation();
 
-	if ((OpA == OpB) && EqualOverrideIds(A.GetNodeID(), B.GetNodeID()))
+	if (EqualOverrideIds(A.GetNodeID(), B.GetNodeID()))
 	{
 		const TArray<FOverriddenPropertyNode>& SubNodesA = A.GetSubPropertyNodes();
 		const TArray<FOverriddenPropertyNode>& SubNodesB = B.GetSubPropertyNodes();
 
-		if (SubNodesA.Num() == SubNodesB.Num())
+		if (OpA == OpB)
 		{
-			int Index = 0;
-
-			for (const FOverriddenPropertyNode& SubNodeA : SubNodesA)
+			if (SubNodesA.Num() == SubNodesB.Num())
 			{
-				// The sub-nodes might be in different order so find a match. This
-				// could be optimized, however, the sub-node count is generally low.
-				const FOverriddenPropertyNode* SubNodeB = nullptr;
-				for (const FOverriddenPropertyNode& Test : SubNodesB)
+				int Index = 0;
+
+				for (const FOverriddenPropertyNode& SubNodeA : SubNodesA)
 				{
-					if (EqualOverrideIds(Test.GetNodeID(), SubNodeA.GetNodeID()))
+					// The sub-nodes might be in different order so find a match. This
+					// could be optimized, however, the sub-node count is generally low.
+					const FOverriddenPropertyNode* SubNodeB = nullptr;
+					for (const FOverriddenPropertyNode& Test : SubNodesB)
 					{
-						SubNodeB = &Test;
-						break;
+						if (EqualOverrideIds(Test.GetNodeID(), SubNodeA.GetNodeID()))
+						{
+							SubNodeB = &Test;
+							break;
+						}
 					}
-				}
 
-				if (SubNodeB)
-				{
-					if (DiffOverrideNodes(SubNodeA, *SubNodeB))
+					if (SubNodeB)
+					{
+						FDiffNodeContext NewCtx = MakeNewContext(SubNodeA, Ctx);
+						if (DiffOverrideNodes(SubNodeA, *SubNodeB, NewCtx))
+						{
+							return true;
+						}
+					}
+					else
 					{
 						return true;
 					}
 				}
-				else
-				{
-					return true;
-				}
+			}
+			else
+			{
+				return true;
+			}
+		}
+		// Don't consider semantically identical overridden state betwen 'all overridden modified' and 'replaced'
+		// TPS doesn't reconstruct the state on load and hence needs to mark struct only modified in all overridden state
+		// to handle schema changes
+		else if (OpA == EOverriddenPropertyOperation::Modified && 
+			OpB == EOverriddenPropertyOperation::Replace)
+		{
+			EOverriddenState State = FOverridableManager::Get().GetOverriddenState(
+				const_cast<UObject*>(Ctx.ObjectA),
+				{},												// consider all properties
+				Ctx.VisitorPath,
+				nullptr,
+				true, 											// ignore added
+				[](TNotNull<const UObject*>) { return false; }  // ignore subobjects
+			);
+			if (State != EOverriddenState::AllOverridden)
+			{
+				return true;
 			}
 		}
 		else
@@ -708,11 +771,11 @@ static bool DiffOverrideNodes(const FOverriddenPropertyNode& A, const FOverridde
 	return false;
 }
 
-static bool DiffOverrides(const FOverriddenPropertySet* A, const FOverriddenPropertySet* B)
+static bool DiffOverrides(const FOverriddenPropertySet* A, const FOverriddenPropertySet* B, FDiffNodeContext& Ctx)
 {
 	if (A && B)
 	{
-		return DiffOverrideNodes(A->GetRootOverriddenPropertyNode(), B->GetRootOverriddenPropertyNode());
+		return DiffOverrideNodes(A->GetRootOverriddenPropertyNode(), B->GetRootOverriddenPropertyNode(), Ctx);
 	}
 	
 	// When both overrides are NULL then the object does not support
@@ -752,7 +815,8 @@ bool DiffObjects(FDiffObjectContext& Ctx)
 
 		Ctx.Visited.Reset();
 
-		if (DiffOverrides(OverrideManager.GetOverriddenProperties(A), OverrideManager.GetOverriddenProperties(*B)))
+		FDiffNodeContext NodeCtx(A);
+		if (DiffOverrides(OverrideManager.GetOverriddenProperties(A), OverrideManager.GetOverriddenProperties(*B), NodeCtx))
 		{
 			Ctx.OverrideDiffs.Emplace(A, *B);
 		}
