PlainProps: Replace incorrect checkSlow in the debug version of DiffMap with correct handling of diffing pairs.
Add local lambda VerifyDebugDiff to verify that the debug version produces the correct result.
Fix diff implementations for animation bindings to diff references rather than diffing by value.

#rnx

#Synced-CL 52611489

diff --git a/Engine/Plugins/Experimental/PlainPropsEngine/Source/Private/PlainPropsAnimationBindings.cpp b/Engine/Plugins/Experimental/PlainPropsEngine/Source/Private/PlainPropsAnimationBindings.cpp
--- a/Engine/Plugins/Experimental/PlainPropsEngine/Source/Private/PlainPropsAnimationBindings.cpp
+++ b/Engine/Plugins/Experimental/PlainPropsEngine/Source/Private/PlainPropsAnimationBindings.cpp
@@ -77,7 +77,7 @@ struct FAnimationAttributeIdentifierBinding : ICustomBinding
 		Dst.*_ScriptStruct = Cast<UScriptStruct>((Dst.*_ScriptStructPath).ResolveObject());
 	}
 
-	inline static bool Diff(FAnimationAttributeIdentifier A, FAnimationAttributeIdentifier B, const FBindContext&)
+	inline static bool Diff(const FAnimationAttributeIdentifier& A, const FAnimationAttributeIdentifier& B, const FBindContext&)
 	{
 		return A != B;
 	}
@@ -170,7 +170,7 @@ struct FAttributeCurveBinding : ICustomBinding
 		Dst.*_bShouldInterpolate = ::UE::Anim::AttributeTypes::CanInterpolateType(ScriptStruct);
 	}
 
-	inline static bool Diff(FAttributeCurve A, FAttributeCurve B, const FBindContext&)
+	inline static bool Diff(const FAttributeCurve& A, const FAttributeCurve& B, const FBindContext&)
 	{
 		using namespace FAttributeCurve_Private;
 
@@ -187,8 +187,8 @@ struct FAttributeCurveBinding : ICustomBinding
 		{
 			return true;
 		}
-		TArray<FAttributeKey>& KeysA = A.*_Keys;
-		TArray<FAttributeKey>& KeysB = B.*_Keys;
+		const TArray<FAttributeKey>& KeysA = A.*_Keys;
+		const TArray<FAttributeKey>& KeysB = B.*_Keys;
 		if (KeysA.Num() != KeysB.Num())
 		{
 			return true;
diff --git a/Engine/Plugins/Experimental/PlainPropsUObject/Source/Private/PlainPropsUObjectBindings.cpp b/Engine/Plugins/Experimental/PlainPropsUObject/Source/Private/PlainPropsUObjectBindings.cpp
--- a/Engine/Plugins/Experimental/PlainPropsUObject/Source/Private/PlainPropsUObjectBindings.cpp
+++ b/Engine/Plugins/Experimental/PlainPropsUObject/Source/Private/PlainPropsUObjectBindings.cpp
@@ -2282,12 +2282,24 @@ inline bool DiffSet(const BindingType& Binding, const ScriptType& A, const Scrip
 	return false;
 }
 
+// FDiffContext version of DiffSet that tries to generate extra debug info for nested differences inside TMaps
 template<typename BindingType, typename ScriptType>
 bool DiffMap(const BindingType& Binding, const ScriptType& A, const ScriptType& B, FDiffContext& Ctx)
 {
+	auto VerifyDebugDiff = [&Binding, &A, &B](bool bDebugDiff) -> bool
+	{
+#if DO_CHECK
+		bool bActualDiff = DiffSet<BindingType, ScriptType>(Binding, A, B);
+		check(bDebugDiff == bActualDiff);
+		return bActualDiff;
+#else
+		return bDebugDiff;
+#endif
+	};
+
 	if (A.NumUnchecked() != B.NumUnchecked())
 	{
-		return true;
+		return VerifyDebugDiff(true);
 	}
 
 	if (A.NumUnchecked() > 0)
@@ -2302,25 +2314,28 @@ bool DiffMap(const BindingType& Binding, const ScriptType& A, const ScriptType&
 			int32 IdxB = Binding.FindKey(B, KeyA);
 			if (IdxB == INDEX_NONE)
 			{
-				// Todo: Add specific diff entry for missing key.
-				return true;
+				return VerifyDebugDiff(true);
 			}
 
 			const void* PairA = KeyA;
 			const void* PairB = B.GetData(IdxB, Binding.Layout);
 			FBindId ValueId = Binding.PairRange.SaveId;
-			// Diff the whole Pair to generate the full nested diff path for diffing values
+			// Diff the whole Pair to try to generate extra nested diffs for the different values
 			if (DiffStructs(PairA, PairB, ValueId, Ctx))
 			{
-				check(!Binding.HasItem(B, PairA));
-				return true;
+				return VerifyDebugDiff(true);
 			}
 
-			checkSlow(Binding.HasItem(B, PairA));
+			// Fallback to the DiffSet version.
+			// This is required since custom bindings may diff a more restricted subset of members than Identical.
+			if (!Binding.HasItem(B, PairA))
+			{
+				return VerifyDebugDiff(true);
+			}
 		}
 	}
 
-	return false;
+	return VerifyDebugDiff(false);
 }
 
 ////////////////////////////////////////////////////////////////////////////////////////////////
