PlainProps: Add template constructors for FMemberSpec and TOptionalId to handle implicit conversions from FDualStructId.
Remove explicit conversions in custom binding code when assigning the optional super and the member specs.

#rnx

#Synced-CL 52611489

diff --git a/Engine/Plugins/Experimental/PlainProps/Source/Public/PlainPropsSpecify.h b/Engine/Plugins/Experimental/PlainProps/Source/Public/PlainPropsSpecify.h
--- a/Engine/Plugins/Experimental/PlainProps/Source/Public/PlainPropsSpecify.h
+++ b/Engine/Plugins/Experimental/PlainProps/Source/Public/PlainPropsSpecify.h
@@ -42,7 +42,8 @@ public:
 	FMemberSpec(ESpecChar In)	: FMemberSpec(ELeafType::Unicode,	static_cast<ELeafWidth>(In), NoId) {}
 	FMemberSpec(ESpecFloat In)	: FMemberSpec(ELeafType::Float,		static_cast<ELeafWidth>(In), NoId) {}
 	
-	FMemberSpec(FDeclId Id)			: FMemberSpec(FStructType{EMemberKind::Struct, /* IsDynamic */ 0, /* IsSuper */ 0}, FInnerId(Id)) {}
+	template<typename T> requires (std::is_convertible_v<T, FDeclId>)
+	FMemberSpec(T Id)				: FMemberSpec(FStructType{EMemberKind::Struct, /* IsDynamic */ 0, /* IsSuper */ 0}, FInnerId(FDeclId(Id))) {}
 	FMemberSpec(ESpecDynamicStruct) : FMemberSpec(FStructType{EMemberKind::Struct, /* IsDynamic */ 1, /* IsSuper */ 0}, NoId) {}
 
 	PLAINPROPS_API FMemberSpec(ESizeType MaxSize, FMemberSpec Inner);
diff --git a/Engine/Plugins/Experimental/PlainProps/Source/Public/PlainPropsTypes.h b/Engine/Plugins/Experimental/PlainProps/Source/Public/PlainPropsTypes.h
--- a/Engine/Plugins/Experimental/PlainProps/Source/Public/PlainPropsTypes.h
+++ b/Engine/Plugins/Experimental/PlainProps/Source/Public/PlainPropsTypes.h
@@ -312,7 +312,12 @@ class TOptionalId : public FOptionalId
 public:
 	constexpr TOptionalId() = default;
 	constexpr TOptionalId(FNoId) {}
-    constexpr TOptionalId(T Id) { Idx = ToIdx(Id); }
+
+	template<typename U>
+	constexpr TOptionalId(U Id) requires (std::is_convertible_v<U, T>)
+	{
+		Idx = ToIdx(static_cast<T>(Id));
+	}
 
 	template<typename U>
 	constexpr explicit TOptionalId(TOptionalId<U> In) requires (std::is_constructible_v<T,U>)
diff --git a/Engine/Plugins/Experimental/PlainPropsEngine/Source/Private/PlainPropsAnimationBindings.cpp b/Engine/Plugins/Experimental/PlainPropsEngine/Source/Private/PlainPropsAnimationBindings.cpp
--- a/Engine/Plugins/Experimental/PlainPropsEngine/Source/Private/PlainPropsAnimationBindings.cpp
+++ b/Engine/Plugins/Experimental/PlainPropsEngine/Source/Private/PlainPropsAnimationBindings.cpp
@@ -49,10 +49,10 @@ struct FAnimationAttributeIdentifierBinding : ICustomBinding
 	FAnimationAttributeIdentifierBinding(TPropertySpecifier<4>& Spec)
 	: MemberIds{Ids::IndexMember("Name"), Ids::IndexMember("BoneName"), Ids::IndexMember("BoneIndex"), Ids::IndexMember("ScriptStructPath")}
 	{
-		Spec.Members[0] = FDeclId(GUE.Structs.Name);
-		Spec.Members[1] = FDeclId(GUE.Structs.Name);
+		Spec.Members[0] = GUE.Structs.Name;
+		Spec.Members[1] = GUE.Structs.Name;
 		Spec.Members[2] = Specify<int32>();
-		Spec.Members[3] = FDeclId(GUE.Structs.SoftObjectPath);
+		Spec.Members[3] = GUE.Structs.SoftObjectPath;
 	}
 
 	void Save(FMemberBuilder& Dst, const FAnimationAttributeIdentifier& Src, const FAnimationAttributeIdentifier* Default, const FSaveContext& Ctx) const
@@ -95,8 +95,8 @@ struct FAttributeCurveBinding : ICustomBinding
 	: AttributeKeyId(GUE.DynamicIds.GetStruct(FAttributeKey::StaticStruct()))
 	, MemberIds{Ids::IndexMember("Keys"), Ids::IndexMember("ScriptStructPath"), Ids::IndexMember("DynamicKeys")}
 	{
-		Spec.Members[0] = DefaultRangeOf(FDeclId(AttributeKeyId));
-		Spec.Members[1] = FDeclId(GUE.Structs.SoftObjectPath);
+		Spec.Members[0] = DefaultRangeOf(AttributeKeyId);
+		Spec.Members[1] = GUE.Structs.SoftObjectPath;
 		Spec.Members[2] = DefaultRangeOf(SpecDynamicStruct);
 	}
 
diff --git a/Engine/Plugins/Experimental/PlainPropsEngine/Source/Private/PlainPropsEngineBindings.cpp b/Engine/Plugins/Experimental/PlainPropsEngine/Source/Private/PlainPropsEngineBindings.cpp
--- a/Engine/Plugins/Experimental/PlainPropsEngine/Source/Private/PlainPropsEngineBindings.cpp
+++ b/Engine/Plugins/Experimental/PlainPropsEngine/Source/Private/PlainPropsEngineBindings.cpp
@@ -72,7 +72,7 @@ struct FTexture2DCook : ICustomBinding
 		Spec.Members[1] = Specify<TextureAddress>(TextureAddressId);
 		Spec.Members[2] = FMemberSpec(IntPointId);
 		Spec.Members[3] = DefaultRangeOf(Specify<EPixelFormat>(GUE.EnumIds.PixelFormat));
-		Spec.Members[4] = DefaultRangeOf(FDeclId(PlatformTextureDataId));
+		Spec.Members[4] = DefaultRangeOf(PlatformTextureDataId);
 	}
 	
 	inline void Save(FMemberBuilder& Dst, const UTexture2D& Src, const UTexture2D* Default, const FSaveContext& Ctx) const
diff --git a/Engine/Plugins/Experimental/PlainPropsEngine/Source/Private/PlainPropsNiagaraBindings.cpp b/Engine/Plugins/Experimental/PlainPropsEngine/Source/Private/PlainPropsNiagaraBindings.cpp
--- a/Engine/Plugins/Experimental/PlainPropsEngine/Source/Private/PlainPropsNiagaraBindings.cpp
+++ b/Engine/Plugins/Experimental/PlainPropsEngine/Source/Private/PlainPropsNiagaraBindings.cpp
@@ -45,7 +45,7 @@ struct FNiagaraTypeDefinitionBinding : ICustomBinding
 	, TypeFlagsId(GetEnumId<Ids, FTypeFlags>())
 	, MemberIds{Ids::IndexMember("ClassStructOrEnum"), Ids::IndexMember("UnderlyingType"), Ids::IndexMember("Flags")}
 	{
-		Spec.Members[0] = FDeclId(GUE.Structs.ObjectPtr);
+		Spec.Members[0] = GUE.Structs.ObjectPtr;
 		Spec.Members[1] = Specify<FUnderlyingType>(UnderlyingTypeId);
 		Spec.Members[2] = Specify<FTypeFlags>(TypeFlagsId);
 	}
@@ -90,7 +90,7 @@ struct FNiagaraTypeDefinitionHandleBinding : ICustomBinding
 	: NiagaraTypeDefinitionId(GUE.DynamicIds.GetStruct(FNiagaraTypeDefinition::StaticStruct()))
 	, MemberIds{Ids::IndexMember("TypeDef")}
 	{
-		Spec.Members[0] = FDeclId(NiagaraTypeDefinitionId);
+		Spec.Members[0] = NiagaraTypeDefinitionId;
 	}
 
 	void Save(FMemberBuilder& Dst, const FNiagaraTypeDefinitionHandle& Src, const FNiagaraTypeDefinitionHandle* Default, const FSaveContext& Ctx) const
diff --git a/Engine/Plugins/Experimental/PlainPropsUObject/Source/Private/PlainPropsUObjectBindings.cpp b/Engine/Plugins/Experimental/PlainPropsUObject/Source/Private/PlainPropsUObjectBindings.cpp
--- a/Engine/Plugins/Experimental/PlainPropsUObject/Source/Private/PlainPropsUObjectBindings.cpp
+++ b/Engine/Plugins/Experimental/PlainPropsUObject/Source/Private/PlainPropsUObjectBindings.cpp
@@ -496,8 +496,8 @@ static void LoadNames(TArray<FName>& Dst, FStructRangeLoadView Src)
 FFieldPathBinding::FFieldPathBinding(TPropertySpecifier<2>& Spec)
 : MemberIds{GUE.Members.Path, GUE.Members.Owner}
 {
-	Spec.Members[0] = DefaultRangeOf(FDeclId(GUE.Structs.Name));
-	Spec.Members[1] = FDeclId(GUE.Structs.WeakObjectPtr);
+	Spec.Members[0] = DefaultRangeOf(GUE.Structs.Name);
+	Spec.Members[1] = GUE.Structs.WeakObjectPtr;
 }
 
 void FFieldPathBinding::Save(FMemberBuilder& Dst, const FFieldPath& Src, const FFieldPath*, const FSaveContext& Ctx) const
@@ -608,7 +608,7 @@ static FTypedRange SaveInvocations(const FMulticastInvocationList& In, const FSa
 
 static FStructDeclarationPtr DeclareMulticastDelegate()
 {
-	return Declare({GUE.Structs.MulticastDelegate, NoId, 0, EMemberPresence::RequireAll, {GUE.Members.Invocations}, {DefaultRangeOf(FDeclId(GUE.Structs.Delegate))}});
+	return Declare({GUE.Structs.MulticastDelegate, NoId, 0, EMemberPresence::RequireAll, {GUE.Members.Invocations}, {DefaultRangeOf(GUE.Structs.Delegate)}});
 }
 
 static void SaveMulticastDelegate(FMemberBuilder& Dst, const FMulticastScriptDelegate& Src, const FSaveContext& Ctx)
@@ -670,7 +670,7 @@ static bool DiffMulticastDelegate(const FMulticastScriptDelegate& A, const FMult
 FMulticastInlineDelegateBinding::FMulticastInlineDelegateBinding(TPropertySpecifier<1>& Spec)
 : MemberIds{GUE.Members.Invocations}
 {
-	Spec.Members[0] = DefaultRangeOf(FDeclId(GUE.Structs.Delegate));
+	Spec.Members[0] = DefaultRangeOf(GUE.Structs.Delegate);
 }
 
 void FMulticastInlineDelegateBinding::Save(FMemberBuilder& Dst, const FMulticastScriptDelegate& Src, const FMulticastScriptDelegate* Default, const FSaveContext& Ctx) const
@@ -971,7 +971,7 @@ bool FReferencePropertyBinding::Diff(const FReferencePropertyValue& A, const FRe
 FInstancedStructBinding::FInstancedStructBinding(TPropertySpecifier<2>& Spec)
 	: MemberIds{ GUE.Names.NameMember("ScriptStruct"), GUE.Names.NameMember("ScriptMemory")}
 {
-	Spec.Members[0] = FDeclId(GUE.Structs.ObjectPtr);
+	Spec.Members[0] = GUE.Structs.ObjectPtr;
 	Spec.Members[1] = SpecDynamicStruct;
 }
 
diff --git a/Engine/Plugins/Experimental/PlainPropsUObject/Source/Private/PlainPropsUObjectMetaBindings.cpp b/Engine/Plugins/Experimental/PlainPropsUObject/Source/Private/PlainPropsUObjectMetaBindings.cpp
--- a/Engine/Plugins/Experimental/PlainPropsUObject/Source/Private/PlainPropsUObjectMetaBindings.cpp
+++ b/Engine/Plugins/Experimental/PlainPropsUObject/Source/Private/PlainPropsUObjectMetaBindings.cpp
@@ -85,7 +85,7 @@ FPropertyBinding::FPropertyBinding(TCustomSpecifier<Ids, NumMembers>& Spec)
 	, LifetimeConditionId(Ids::IndexEnum(GUE.Scopes.CoreUObject, "ELifetimeCondition"))
 {
 	// All FFields are FProperties so we can fold all FField members here for simplicity and efficiency
-	Spec.Members[M::NamePrivate] = FDeclId(GUE.Structs.Name);
+	Spec.Members[M::NamePrivate] = GUE.Structs.Name;
 	Spec.Members[M::ArrayDim] = Specify<int32>();
 	Spec.Members[M::ElementSize] = Specify<int32>();
 	Spec.Members[M::PropertyFlags] = Specify<EPropertyFlags>(PropertyFlagsId);
@@ -129,7 +129,7 @@ void FPropertyBinding::Load(FProperty& Dst, FStructLoadView Src, ECustomLoadMeth
 
 FBoolPropertyBinding::FBoolPropertyBinding(TCustomSpecifier<Ids, NumMembers>& Spec)
 {
-	Spec.Super = FDeclId(GUE.Properties.Property);
+	Spec.Super = GUE.Properties.Property;
 	Spec.Members[M::FieldSize] = Specify<uint8>();
 	Spec.Members[M::ByteOffset] = Specify<uint8>();
 	Spec.Members[M::ByteMask] = Specify<uint8>();
@@ -160,8 +160,8 @@ void FBoolPropertyBinding::Load(FBoolProperty& Dst, FStructLoadView Src, ECustom
 
 FStructPropertyBinding::FStructPropertyBinding(TCustomSpecifier<Ids, NumMembers>& Spec)
 {
-	Spec.Super = FDeclId(GUE.Properties.Property);
-	Spec.Members[M::Struct] = FDeclId(GUE.Structs.ObjectPtr);
+	Spec.Super = GUE.Properties.Property;
+	Spec.Members[M::Struct] = GUE.Structs.ObjectPtr;
 }
 
 void FStructPropertyBinding::Save(FMemberBuilder& Dst, const FStructProperty& Src, const FStructProperty*, const FSaveContext& Ctx) const
@@ -180,8 +180,8 @@ void FStructPropertyBinding::Load(FStructProperty& Dst, FStructLoadView Src, ECu
 
 FClassPropertyBinding::FClassPropertyBinding(TCustomSpecifier<Ids, NumMembers>& Spec)
 {
-	Spec.Super = FDeclId(GUE.Properties.Property);
-	Spec.Members[M::MetaClass] = FDeclId(GUE.Structs.ObjectPtr);
+	Spec.Super = GUE.Properties.Property;
+	Spec.Members[M::MetaClass] = GUE.Structs.ObjectPtr;
 }
 
 void FClassPropertyBinding::Save(FMemberBuilder& Dst, const FClassProperty& Src, const FClassProperty*, const FSaveContext& Ctx) const
@@ -207,8 +207,8 @@ void FClassPropertyBinding::Load(FClassProperty& Dst, FStructLoadView Src, ECust
 
 FObjectPropertyBaseBinding::FObjectPropertyBaseBinding(TCustomSpecifier<Ids, NumMembers>& Spec)
 {
-	Spec.Super = FDeclId(GUE.Properties.Property);
-	Spec.Members[M::PropertyClass] = FDeclId(GUE.Structs.ObjectPtr);
+	Spec.Super = GUE.Properties.Property;
+	Spec.Members[M::PropertyClass] = GUE.Structs.ObjectPtr;
 }
 
 void FObjectPropertyBaseBinding::Save(FMemberBuilder& Dst, const FObjectPropertyBase& Src, const FObjectPropertyBase*, const FSaveContext& Ctx) const
@@ -227,9 +227,9 @@ void FObjectPropertyBaseBinding::Load(FObjectPropertyBase& Dst, FStructLoadView
 
 FArrayPropertyBinding::FArrayPropertyBinding(TCustomSpecifier<Ids, NumMembers>& Spec)
 {
-	Spec.Super = FDeclId(GUE.Properties.Property);
+	Spec.Super = GUE.Properties.Property;
 	// Intentionally skip ArrayFlags (no support for blueprint UsesMemoryImageAllocator)
-	Spec.Members[M::InnerPropertyTypeName] = FDeclId(GUE.Structs.Name);
+	Spec.Members[M::InnerPropertyTypeName] = GUE.Structs.Name;
 	Spec.Members[M::InnerProperty] = SpecDynamicStruct;
 }
 
@@ -252,8 +252,8 @@ void FArrayPropertyBinding::Load(FArrayProperty& Dst, FStructLoadView Src, ECust
 
 FSetPropertyBinding::FSetPropertyBinding(TCustomSpecifier<Ids, NumMembers>& Spec)
 {
-	Spec.Super = FDeclId(GUE.Properties.Property);
-	Spec.Members[M::InnerPropertyTypeName] = FDeclId(GUE.Structs.Name);
+	Spec.Super = GUE.Properties.Property;
+	Spec.Members[M::InnerPropertyTypeName] = GUE.Structs.Name;
 	Spec.Members[M::InnerProperty] = SpecDynamicStruct;
 }
 
@@ -276,10 +276,10 @@ void FSetPropertyBinding::Load(FSetProperty& Dst, FStructLoadView Src, ECustomLo
 
 FMapPropertyBinding::FMapPropertyBinding(TCustomSpecifier<Ids, NumMembers>& Spec)
 {
-	Spec.Super = FDeclId(GUE.Properties.Property);
-	Spec.Members[M::KeyPropertyTypeName] = FDeclId(GUE.Structs.Name);
+	Spec.Super = GUE.Properties.Property;
+	Spec.Members[M::KeyPropertyTypeName] = GUE.Structs.Name;
 	Spec.Members[M::KeyProperty] = SpecDynamicStruct;
-	Spec.Members[M::ValuePropertyTypeName] = FDeclId(GUE.Structs.Name);
+	Spec.Members[M::ValuePropertyTypeName] = GUE.Structs.Name;
 	Spec.Members[M::ValueProperty] = SpecDynamicStruct;
 }
 
@@ -304,8 +304,8 @@ void FMapPropertyBinding::Load(FMapProperty& Dst, FStructLoadView Src, ECustomLo
 
 FOptionalPropertyBinding::FOptionalPropertyBinding(TCustomSpecifier<Ids, NumMembers>& Spec)
 {
-	Spec.Super = FDeclId(GUE.Properties.Property);
-	Spec.Members[M::ValuePropertyTypeName] = FDeclId(GUE.Structs.Name);
+	Spec.Super = GUE.Properties.Property;
+	Spec.Members[M::ValuePropertyTypeName] = GUE.Structs.Name;
 	Spec.Members[M::ValueProperty] = SpecDynamicStruct;
 }
 
@@ -329,8 +329,8 @@ void FOptionalPropertyBinding::Load(FOptionalProperty& Dst, FStructLoadView Src,
 
 FBytePropertyBinding::FBytePropertyBinding(TCustomSpecifier<Ids, NumMembers>& Spec)
 {
-	Spec.Super = FDeclId(GUE.Properties.Property);
-	Spec.Members[M::Enum] = FDeclId(GUE.Structs.ObjectPtr);
+	Spec.Super = GUE.Properties.Property;
+	Spec.Members[M::Enum] = GUE.Structs.ObjectPtr;
 }
 
 void FBytePropertyBinding::Save(FMemberBuilder& Dst, const FByteProperty& Src, const FByteProperty* Default, const FSaveContext& Ctx) const
@@ -349,10 +349,10 @@ void FBytePropertyBinding::Load(FByteProperty& Dst, FStructLoadView Src, ECustom
 
 FEnumPropertyBinding::FEnumPropertyBinding(TCustomSpecifier<Ids, NumMembers>& Spec)
 {
-	Spec.Super = FDeclId(GUE.Properties.Property);
-	Spec.Members[M::UnderlyingPropertyTypeName] = FDeclId(GUE.Structs.Name);
+	Spec.Super = GUE.Properties.Property;
+	Spec.Members[M::UnderlyingPropertyTypeName] = GUE.Structs.Name;
 	Spec.Members[M::UnderlyingProperty] = SpecDynamicStruct;
-	Spec.Members[M::Enum] = FDeclId(GUE.Structs.ObjectPtr);
+	Spec.Members[M::Enum] = GUE.Structs.ObjectPtr;
 }
 
 void FEnumPropertyBinding::Save(FMemberBuilder& Dst, const FEnumProperty& Src, const FEnumProperty* Default, const FSaveContext& Ctx) const
@@ -383,12 +383,12 @@ void FEnumPropertyBinding::Load(FEnumProperty& Dst, FStructLoadView Src, ECustom
 
 UStructBinding::UStructBinding(TCustomSpecifier<Ids, NumMembers>& Spec)
 {
-	Spec.Members[M::SuperStruct] = FDeclId(GUE.Structs.ObjectPtr);
+	Spec.Members[M::SuperStruct] = GUE.Structs.ObjectPtr;
 	// Children is an array of UFields, these are always UFunctions
-	Spec.Members[M::Children] = DefaultRangeOf(FDeclId(GUE.Structs.ObjectPtr));
+	Spec.Members[M::Children] = DefaultRangeOf(GUE.Structs.ObjectPtr);
 	// This UStruct base binding saves and loads the linked list of FProperties as an array.
 	// The derived meta classes saves and loads the actual FProperty data as named schema members.
-	Spec.Members[M::ChildPropertyClasses] = DefaultRangeOf(FDeclId(GUE.Structs.Name));
+	Spec.Members[M::ChildPropertyClasses] = DefaultRangeOf(GUE.Structs.Name);
 }
 
 FTypedRange UStructBinding::SaveChildrenStatic(const UField* Children, int32 Num, const FSaveContext& Ctx)
@@ -521,14 +521,14 @@ UClassBinding::UClassBinding(TCustomSpecifier<Ids, NumMembers>& Spec)
 : ClassFlagsId(GetEnumId<Ids, EClassFlags>())
 , ImplementedInterfaceId(GetStructDualId<Ids, FImplementedInterface>())
 {
-	Spec.Super = FDeclId(GUE.Structs.Struct);
+	Spec.Super = GUE.Structs.Struct;
 	Spec.Members[M::ClassFlags] = Specify<EClassFlags>(ClassFlagsId);
-	Spec.Members[M::ClassWithin] = FDeclId(GUE.Structs.ObjectPtr);
-	Spec.Members[M::ClassConfigName] = FDeclId(GUE.Structs.Name);
+	Spec.Members[M::ClassWithin] = GUE.Structs.ObjectPtr;
+	Spec.Members[M::ClassConfigName] = GUE.Structs.Name;
 #if WITH_EDITORONLY_DATA
-	Spec.Members[M::ClassGeneratedBy] = FDeclId(GUE.Structs.ObjectPtr);
+	Spec.Members[M::ClassGeneratedBy] = GUE.Structs.ObjectPtr;
 #endif
-	Spec.Members[M::Interfaces] = DefaultRangeOf(FDeclId(ImplementedInterfaceId));
+	Spec.Members[M::Interfaces] = DefaultRangeOf(ImplementedInterfaceId);
 }
 
 FTypedRange UClassBinding::SaveInterfacesStatic(FBindId Id, TConstArrayView<FImplementedInterface> Interfaces, const FSaveContext& Ctx)
@@ -563,7 +563,7 @@ void UClassBinding::Save(FMemberBuilder& Dst, const UClass& Src, const UClass*,
 #endif
 	if (Src.Interfaces.Num())
 	{
-		Dst.AddRange(MemberIds[M::Interfaces], SaveInterfacesStatic(FBindId(ImplementedInterfaceId), Src.Interfaces, Ctx));
+		Dst.AddRange(MemberIds[M::Interfaces], SaveInterfacesStatic(ImplementedInterfaceId, Src.Interfaces, Ctx));
 	}
 }
 
@@ -625,9 +625,9 @@ void UClassBinding::Load(UClass& Dst, FStructLoadView Src, ECustomLoadMethod Met
 UFunctionBinding::UFunctionBinding(TCustomSpecifier<Ids, NumMembers>& Spec)
 : FunctionFlagsId(GetEnumId<Ids, EFunctionFlags>())
 {
-	Spec.Super = FDeclId(GUE.Structs.Struct);
+	Spec.Super = GUE.Structs.Struct;
 	Spec.Members[M::FunctionFlags] = Specify<EFunctionFlags>(FunctionFlagsId);
-	Spec.Members[M::EventGraphFunction] = FDeclId(GUE.Structs.ObjectPtr);
+	Spec.Members[M::EventGraphFunction] = GUE.Structs.ObjectPtr;
 	Spec.Members[M::EventGraphCallOffset] = Specify<int32>();
 }
 
@@ -653,7 +653,7 @@ void UFunctionBinding::Load(UFunction& Dst, FStructLoadView Src, ECustomLoadMeth
 UScriptStructBinding::UScriptStructBinding(TCustomSpecifier<Ids, NumMembers>& Spec)
 : StructFlagsId(GetEnumId<Ids, EStructFlags>())
 {
-	Spec.Super = FDeclId(GUE.Structs.Struct);
+	Spec.Super = GUE.Structs.Struct;
 	Spec.Members[M::StructFlags] = Specify<EStructFlags>(StructFlagsId);
 }
 
@@ -812,9 +812,9 @@ public:
 	{
 		FDualStructId Super = GetSuper(Struct);
 		BindAllMembers(Struct);
-		FStructDeclarationPtr Declaration = Declare(FDeclId(Super));
+		FStructDeclarationPtr Declaration = Declare(Super);
 		// Todo: Ownership / memory leak
-		ICustomBinding* Leak = new FStructBindingMeta(FBindId(Super), Names, Members);
+		ICustomBinding* Leak = new FStructBindingMeta(Super, Names, Members);
 		GUE.Customs.BindStruct(Owner, *Leak, FStructDeclarationPtr(Declaration), {});
 	}
 
@@ -846,8 +846,8 @@ public:
 		{
 			// Todo: non-composable hack, declare CDO for UClass in the generic meta binding
 			FDualStructId Id = GUE.Structs.ObjectPtr;
-			Members.Emplace(FBindId(Id));
-			Specs.Emplace(FDeclId(Id));
+			Members.Emplace(Id);
+			Specs.Emplace(Id);
 			Names.Emplace(GUE.Names.NameMember("ClassDefaultObject"));
 		}
 	}
@@ -855,8 +855,8 @@ public:
 	void BindMember(const FProperty* Property)
 	{
 		FDualStructId Id = IndexPropertyMeta(Property);
-		Members.Emplace(FBindId(Id));
-		Specs.Emplace(FDeclId(Id));
+		Members.Emplace(Id);
+		Specs.Emplace(Id);
 		Names.Emplace(GUE.Names.NameMember(Property->GetFName()));
 	}
 
