PlainProps: Tweak FPropertyBinding to write its TEnumAsByte<ELifetimeCondition> BlueprintRreplicationCondition member as an enum8 rather than enum32.
This is temporary but consistent with the current implementation for schema properties in FPropertyBinder::BindByte.

#rnx

#Synced-CL 52611489

diff --git a/Engine/Plugins/Experimental/PlainPropsUObject/Source/Private/PlainPropsUObjectMetaBindings.cpp b/Engine/Plugins/Experimental/PlainPropsUObject/Source/Private/PlainPropsUObjectMetaBindings.cpp
--- a/Engine/Plugins/Experimental/PlainPropsUObject/Source/Private/PlainPropsUObjectMetaBindings.cpp
+++ b/Engine/Plugins/Experimental/PlainPropsUObject/Source/Private/PlainPropsUObjectMetaBindings.cpp
@@ -82,7 +82,7 @@ static void LoadSingleProperty(FMemberLoader& Members, FProperty*& Dst, const FP
 
 FPropertyBinding::FPropertyBinding(TCustomSpecifier<Ids, NumMembers>& Spec)
 	: PropertyFlagsId(GetEnumId<Ids, EPropertyFlags>())
-	, LifetimeConditionId(Ids::IndexEnum(GUE.Scopes.CoreUObject, "ELifetimeCondition"))
+	, LifetimeConditionId(Ids::IndexEnum(GUE.Scopes.CoreUObject, "ELifetimeCondition")) // too early for GUE.DynamicIds.GetEnum
 {
 	// All FFields are FProperties so we can fold all FField members here for simplicity and efficiency
 	Spec.Members[M::NamePrivate] = GUE.Structs.Name;
@@ -90,7 +90,8 @@ FPropertyBinding::FPropertyBinding(TCustomSpecifier<Ids, NumMembers>& Spec)
 	Spec.Members[M::ElementSize] = Specify<int32>();
 	Spec.Members[M::PropertyFlags] = Specify<EPropertyFlags>(PropertyFlagsId);
 	Spec.Members[M::RepIndex] = Specify<uint16>();
-	Spec.Members[M::BlueprintReplicationCondition] = Specify<ELifetimeCondition>(LifetimeConditionId);
+	// Todo: Cleanup TEnumAsByte reflection, for now specify/save/load the same enum8 type as the schema version, see FPropertyBinder::BindByte
+	Spec.Members[M::BlueprintReplicationCondition] = FMemberSpec(ELeafType::Enum, ELeafWidth::B8, LifetimeConditionId);
 	// Todo: Save and load editor only FField::MetaDataMap
 	// TMap<FName, FString>* MetaDataMap;
 }
@@ -104,8 +105,7 @@ void FPropertyBinding::Save(FMemberBuilder& Dst, const FProperty& Src, const FPr
 	EPropertyFlags SaveFlags = Src.PropertyFlags & ~CPF_ComputedFlags;
 	Dst.AddEnum(MemberIds[M::PropertyFlags], PropertyFlagsId, SaveFlags);
 	Dst.Add(MemberIds[M::RepIndex], Src.RepIndex);
-	ELifetimeCondition SaveBlueprintReplicationCondition = (Src.*_BlueprintReplicationCondition).GetValue();
-	Dst.AddEnum(MemberIds[M::BlueprintReplicationCondition], LifetimeConditionId, SaveBlueprintReplicationCondition);
+	Dst.AddEnum<uint8>(MemberIds[M::BlueprintReplicationCondition], LifetimeConditionId, (Src.*_BlueprintReplicationCondition).GetIntValue());
 	// TMap<FName, FString>* MetaDataMap;
 }
 
@@ -119,8 +119,7 @@ void FPropertyBinding::Load(FProperty& Dst, FStructLoadView Src, ECustomLoadMeth
 	EPropertyFlags LoadFlags = Members.GrabLeaf().As<EPropertyFlags>();
 	Dst.PropertyFlags = (LoadFlags & ~CPF_ComputedFlags) | (Dst.PropertyFlags & CPF_ComputedFlags);
 	Dst.RepIndex = Members.GrabLeaf().As<uint16>();
-	ELifetimeCondition LoadBlueprintReplicationCondition = Members.GrabLeaf().As<ELifetimeCondition>();
-	Dst.*_BlueprintReplicationCondition = LoadBlueprintReplicationCondition;
+	Dst.*_BlueprintReplicationCondition = TEnumAsByte<ELifetimeCondition>(Members.GrabLeaf().AsUnderlyingValue<uint8>());
 	// TMap<FName, FString>* MetaDataMap;
 	Dst.*_Offset_Internal = 0;
 	Dst.DestructorLinkNext = nullptr;
