PlainProps: Replace TCustomSpecifier with TPropertySpecifier in MetaBindings.

#rnx
#Synced-CL 52611489

diff --git a/Engine/Plugins/Experimental/PlainPropsUObject/Source/Internal/PlainPropsUObjectMetaBindingsInternal.h b/Engine/Plugins/Experimental/PlainPropsUObject/Source/Internal/PlainPropsUObjectMetaBindingsInternal.h
--- a/Engine/Plugins/Experimental/PlainPropsUObject/Source/Internal/PlainPropsUObjectMetaBindingsInternal.h
+++ b/Engine/Plugins/Experimental/PlainPropsUObject/Source/Internal/PlainPropsUObjectMetaBindingsInternal.h
@@ -207,7 +207,7 @@ struct FPropertyBinding : ICustomBinding
 {
 	UE_CUSTOM_BIND_MEMBERS(FProperty,
 		NamePrivate, ArrayDim, ElementSize, PropertyFlags, RepIndex, BlueprintReplicationCondition);
-	FPropertyBinding(TCustomSpecifier<Ids, NumMembers>& Spec);
+	FPropertyBinding(TPropertySpecifier<NumMembers>& Spec);
 
 	const FEnumId PropertyFlagsId;
 	const FEnumId LifetimeConditionId;
@@ -220,7 +220,7 @@ struct FPropertyBinding : ICustomBinding
 struct FBoolPropertyBinding : ICustomBinding
 {
 	UE_CUSTOM_BIND_MEMBERS(FBoolProperty, FieldSize, ByteOffset, ByteMask, FieldMask);
-	FBoolPropertyBinding(TCustomSpecifier<Ids, NumMembers>& Spec);
+	FBoolPropertyBinding(TPropertySpecifier<NumMembers>& Spec);
 
 	void		Save(FMemberBuilder& Dst, const Type& Src, const Type* Default, const FSaveContext& Ctx) const;
 	void		Load(Type& Dst, FStructLoadView Src, ECustomLoadMethod Method) const;
@@ -230,7 +230,7 @@ struct FBoolPropertyBinding : ICustomBinding
 struct FStructPropertyBinding : ICustomBinding
 {
 	UE_CUSTOM_BIND_MEMBERS(FStructProperty, Struct);
-	FStructPropertyBinding(TCustomSpecifier<Ids, NumMembers>& Spec);
+	FStructPropertyBinding(TPropertySpecifier<NumMembers>& Spec);
 
 	void		Save(FMemberBuilder& Dst, const Type& Src, const Type* Default, const FSaveContext& Ctx) const;
 	void		Load(Type& Dst, FStructLoadView Src, ECustomLoadMethod Method) const;
@@ -240,7 +240,7 @@ struct FStructPropertyBinding : ICustomBinding
 struct FClassPropertyBinding : ICustomBinding
 {
 	UE_CUSTOM_BIND_MEMBERS(FClassProperty, MetaClass);
-	FClassPropertyBinding(TCustomSpecifier<Ids, NumMembers>& Spec);
+	FClassPropertyBinding(TPropertySpecifier<NumMembers>& Spec);
 
 	void		Save(FMemberBuilder& Dst, const Type& Src, const Type* Default, const FSaveContext& Ctx) const;
 	void		Load(Type& Dst, FStructLoadView Src, ECustomLoadMethod Method) const;
@@ -250,7 +250,7 @@ struct FClassPropertyBinding : ICustomBinding
 struct FObjectPropertyBaseBinding : ICustomBinding
 {
 	UE_CUSTOM_BIND_MEMBERS(FObjectPropertyBase, PropertyClass);
-	FObjectPropertyBaseBinding(TCustomSpecifier<Ids, NumMembers>& Spec);
+	FObjectPropertyBaseBinding(TPropertySpecifier<NumMembers>& Spec);
 
 	void		Save(FMemberBuilder& Dst, const Type& Src, const Type* Default, const FSaveContext& Ctx) const;
 	void		Load(Type& Dst, FStructLoadView Src, ECustomLoadMethod Method) const;
@@ -260,7 +260,7 @@ struct FObjectPropertyBaseBinding : ICustomBinding
 struct FArrayPropertyBinding : ICustomBinding
 {
 	UE_CUSTOM_BIND_MEMBERS(FArrayProperty, InnerPropertyTypeName, InnerProperty);
-	FArrayPropertyBinding(TCustomSpecifier<Ids, NumMembers>& Spec);
+	FArrayPropertyBinding(TPropertySpecifier<NumMembers>& Spec);
 
 	void		Save(FMemberBuilder& Dst, const Type& Src, const Type* Default, const FSaveContext& Ctx) const;
 	void		Load(Type& Dst, FStructLoadView Src, ECustomLoadMethod Method) const;
@@ -270,7 +270,7 @@ struct FArrayPropertyBinding : ICustomBinding
 struct FSetPropertyBinding : ICustomBinding
 {
 	UE_CUSTOM_BIND_MEMBERS(FSetProperty, InnerPropertyTypeName, InnerProperty);
-	FSetPropertyBinding(TCustomSpecifier<Ids, NumMembers>& Spec);
+	FSetPropertyBinding(TPropertySpecifier<NumMembers>& Spec);
 
 	void		Save(FMemberBuilder& Dst, const Type& Src, const Type* Default, const FSaveContext& Ctx) const;
 	void		Load(Type& Dst, FStructLoadView Src, ECustomLoadMethod Method) const;
@@ -280,7 +280,7 @@ struct FSetPropertyBinding : ICustomBinding
 struct FMapPropertyBinding : ICustomBinding
 {
 	UE_CUSTOM_BIND_MEMBERS(FMapProperty, KeyPropertyTypeName, KeyProperty, ValuePropertyTypeName, ValueProperty);
-	FMapPropertyBinding(TCustomSpecifier<Ids, NumMembers>& Spec);
+	FMapPropertyBinding(TPropertySpecifier<NumMembers>& Spec);
 
 	void		Save(FMemberBuilder& Dst, const Type& Src, const Type* Default, const FSaveContext& Ctx) const;
 	void		Load(Type& Dst, FStructLoadView Src, ECustomLoadMethod Method) const;
@@ -290,7 +290,7 @@ struct FMapPropertyBinding : ICustomBinding
 struct FOptionalPropertyBinding : ICustomBinding
 {
 	UE_CUSTOM_BIND_MEMBERS(FOptionalProperty, ValuePropertyTypeName, ValueProperty);
-	FOptionalPropertyBinding(TCustomSpecifier<Ids, NumMembers>& Spec);
+	FOptionalPropertyBinding(TPropertySpecifier<NumMembers>& Spec);
 
 	void		Save(FMemberBuilder& Dst, const Type& Src, const Type* Default, const FSaveContext& Ctx) const;
 	void		Load(Type& Dst, FStructLoadView Src, ECustomLoadMethod Method) const;
@@ -300,7 +300,7 @@ struct FOptionalPropertyBinding : ICustomBinding
 struct FBytePropertyBinding : ICustomBinding
 {
 	UE_CUSTOM_BIND_MEMBERS(FByteProperty, Enum);
-	FBytePropertyBinding(TCustomSpecifier<Ids, NumMembers>& Spec);
+	FBytePropertyBinding(TPropertySpecifier<NumMembers>& Spec);
 
 	void		Save(FMemberBuilder& Dst, const Type& Src, const Type* Default, const FSaveContext& Ctx) const;
 	void		Load(Type& Dst, FStructLoadView Src, ECustomLoadMethod Method) const;
@@ -310,7 +310,7 @@ struct FBytePropertyBinding : ICustomBinding
 struct FEnumPropertyBinding : ICustomBinding
 {
 	UE_CUSTOM_BIND_MEMBERS(FEnumProperty, UnderlyingPropertyTypeName, UnderlyingProperty, Enum);
-	FEnumPropertyBinding(TCustomSpecifier<Ids, NumMembers>& Spec);
+	FEnumPropertyBinding(TPropertySpecifier<NumMembers>& Spec);
 
 	void		Save(FMemberBuilder& Dst, const Type& Src, const Type* Default, const FSaveContext& Ctx) const;
 	void		Load(Type& Dst, FStructLoadView Src, ECustomLoadMethod Method) const;
@@ -323,7 +323,7 @@ struct FEnumPropertyBinding : ICustomBinding
 struct UStructBinding : ICustomBinding
 {
 	UE_CUSTOM_BIND_MEMBERS(UStruct, SuperStruct, Children, ChildPropertyClasses);
-	UStructBinding(TCustomSpecifier<Ids, NumMembers>& Spec);
+	UStructBinding(TPropertySpecifier<NumMembers>& Spec);
 
 	static FTypedRange SaveChildrenStatic(const UField* Children, int32 Num, const FSaveContext& Ctx);
 	static FTypedRange SaveChildPropertyClassesStatic(const FField* ChildProperties, int32 Num, const FSaveContext& Ctx);
@@ -342,7 +342,7 @@ struct UClassBinding : ICustomBinding
 #else
 	UE_CUSTOM_BIND_MEMBERS(UClass, ClassFlags, ClassWithin, ClassConfigName, Interfaces);
 #endif
-	UClassBinding(TCustomSpecifier<Ids, NumMembers>& Spec);
+	UClassBinding(TPropertySpecifier<NumMembers>& Spec);
 
 	const FEnumId ClassFlagsId;
 	const FDualStructId ImplementedInterfaceId;
@@ -358,7 +358,7 @@ struct UClassBinding : ICustomBinding
 struct UFunctionBinding : ICustomBinding
 {
 	UE_CUSTOM_BIND_MEMBERS(UFunction, FunctionFlags, EventGraphFunction, EventGraphCallOffset);
-	UFunctionBinding(TCustomSpecifier<Ids, NumMembers>& Spec);
+	UFunctionBinding(TPropertySpecifier<NumMembers>& Spec);
 
 	const FEnumId FunctionFlagsId;
 
@@ -370,7 +370,7 @@ struct UFunctionBinding : ICustomBinding
 struct UScriptStructBinding : ICustomBinding
 {
 	UE_CUSTOM_BIND_MEMBERS(UScriptStruct, StructFlags);
-	UScriptStructBinding(TCustomSpecifier<Ids, NumMembers>& Spec);
+	UScriptStructBinding(TPropertySpecifier<NumMembers>& Spec);
 
 	const FEnumId StructFlagsId;
 
diff --git a/Engine/Plugins/Experimental/PlainPropsUObject/Source/Private/PlainPropsUObjectMetaBindings.cpp b/Engine/Plugins/Experimental/PlainPropsUObject/Source/Private/PlainPropsUObjectMetaBindings.cpp
--- a/Engine/Plugins/Experimental/PlainPropsUObject/Source/Private/PlainPropsUObjectMetaBindings.cpp
+++ b/Engine/Plugins/Experimental/PlainPropsUObject/Source/Private/PlainPropsUObjectMetaBindings.cpp
@@ -80,7 +80,7 @@ static void LoadSingleProperty(FMemberLoader& Members, FProperty*& Dst, const FP
 	}
 }
 
-FPropertyBinding::FPropertyBinding(TCustomSpecifier<Ids, NumMembers>& Spec)
+FPropertyBinding::FPropertyBinding(TPropertySpecifier<NumMembers>& Spec)
 	: PropertyFlagsId(GetEnumId<Ids, EPropertyFlags>())
 	, LifetimeConditionId(Ids::IndexEnum(GUE.Scopes.CoreUObject, "ELifetimeCondition")) // too early for GUE.DynamicIds.GetEnum
 {
@@ -126,7 +126,7 @@ void FPropertyBinding::Load(FProperty& Dst, FStructLoadView Src, ECustomLoadMeth
 	check(!Members.HasMore());
 }
 
-FBoolPropertyBinding::FBoolPropertyBinding(TCustomSpecifier<Ids, NumMembers>& Spec)
+FBoolPropertyBinding::FBoolPropertyBinding(TPropertySpecifier<NumMembers>& Spec)
 {
 	Spec.Super = GUE.Properties.Property;
 	Spec.Members[M::FieldSize] = Specify<uint8>();
@@ -157,7 +157,7 @@ void FBoolPropertyBinding::Load(FBoolProperty& Dst, FStructLoadView Src, ECustom
 	check(!Members.HasMore());
 }
 
-FStructPropertyBinding::FStructPropertyBinding(TCustomSpecifier<Ids, NumMembers>& Spec)
+FStructPropertyBinding::FStructPropertyBinding(TPropertySpecifier<NumMembers>& Spec)
 {
 	Spec.Super = GUE.Properties.Property;
 	Spec.Members[M::Struct] = GUE.Structs.ObjectPtr;
@@ -177,7 +177,7 @@ void FStructPropertyBinding::Load(FStructProperty& Dst, FStructLoadView Src, ECu
 	check(!Members.HasMore());
 }
 
-FClassPropertyBinding::FClassPropertyBinding(TCustomSpecifier<Ids, NumMembers>& Spec)
+FClassPropertyBinding::FClassPropertyBinding(TPropertySpecifier<NumMembers>& Spec)
 {
 	Spec.Super = GUE.Properties.Property;
 	Spec.Members[M::MetaClass] = GUE.Structs.ObjectPtr;
@@ -204,7 +204,7 @@ void FClassPropertyBinding::Load(FClassProperty& Dst, FStructLoadView Src, ECust
 	check(!Members.HasMore());
 }
 
-FObjectPropertyBaseBinding::FObjectPropertyBaseBinding(TCustomSpecifier<Ids, NumMembers>& Spec)
+FObjectPropertyBaseBinding::FObjectPropertyBaseBinding(TPropertySpecifier<NumMembers>& Spec)
 {
 	Spec.Super = GUE.Properties.Property;
 	Spec.Members[M::PropertyClass] = GUE.Structs.ObjectPtr;
@@ -224,7 +224,7 @@ void FObjectPropertyBaseBinding::Load(FObjectPropertyBase& Dst, FStructLoadView
 	check(!Members.HasMore());
 }
 
-FArrayPropertyBinding::FArrayPropertyBinding(TCustomSpecifier<Ids, NumMembers>& Spec)
+FArrayPropertyBinding::FArrayPropertyBinding(TPropertySpecifier<NumMembers>& Spec)
 {
 	Spec.Super = GUE.Properties.Property;
 	// Intentionally skip ArrayFlags (no support for blueprint UsesMemoryImageAllocator)
@@ -249,7 +249,7 @@ void FArrayPropertyBinding::Load(FArrayProperty& Dst, FStructLoadView Src, ECust
 	check(!Members.HasMore());
 }
 
-FSetPropertyBinding::FSetPropertyBinding(TCustomSpecifier<Ids, NumMembers>& Spec)
+FSetPropertyBinding::FSetPropertyBinding(TPropertySpecifier<NumMembers>& Spec)
 {
 	Spec.Super = GUE.Properties.Property;
 	Spec.Members[M::InnerPropertyTypeName] = GUE.Structs.Name;
@@ -273,7 +273,7 @@ void FSetPropertyBinding::Load(FSetProperty& Dst, FStructLoadView Src, ECustomLo
 	check(!Members.HasMore());
 }
 
-FMapPropertyBinding::FMapPropertyBinding(TCustomSpecifier<Ids, NumMembers>& Spec)
+FMapPropertyBinding::FMapPropertyBinding(TPropertySpecifier<NumMembers>& Spec)
 {
 	Spec.Super = GUE.Properties.Property;
 	Spec.Members[M::KeyPropertyTypeName] = GUE.Structs.Name;
@@ -301,7 +301,7 @@ void FMapPropertyBinding::Load(FMapProperty& Dst, FStructLoadView Src, ECustomLo
 	check(!Members.HasMore());
 }
 
-FOptionalPropertyBinding::FOptionalPropertyBinding(TCustomSpecifier<Ids, NumMembers>& Spec)
+FOptionalPropertyBinding::FOptionalPropertyBinding(TPropertySpecifier<NumMembers>& Spec)
 {
 	Spec.Super = GUE.Properties.Property;
 	Spec.Members[M::ValuePropertyTypeName] = GUE.Structs.Name;
@@ -326,7 +326,7 @@ void FOptionalPropertyBinding::Load(FOptionalProperty& Dst, FStructLoadView Src,
 	check(!Members.HasMore());
 }
 
-FBytePropertyBinding::FBytePropertyBinding(TCustomSpecifier<Ids, NumMembers>& Spec)
+FBytePropertyBinding::FBytePropertyBinding(TPropertySpecifier<NumMembers>& Spec)
 {
 	Spec.Super = GUE.Properties.Property;
 	Spec.Members[M::Enum] = GUE.Structs.ObjectPtr;
@@ -346,7 +346,7 @@ void FBytePropertyBinding::Load(FByteProperty& Dst, FStructLoadView Src, ECustom
 	check(!Members.HasMore());
 }
 
-FEnumPropertyBinding::FEnumPropertyBinding(TCustomSpecifier<Ids, NumMembers>& Spec)
+FEnumPropertyBinding::FEnumPropertyBinding(TPropertySpecifier<NumMembers>& Spec)
 {
 	Spec.Super = GUE.Properties.Property;
 	Spec.Members[M::UnderlyingPropertyTypeName] = GUE.Structs.Name;
@@ -380,7 +380,7 @@ void FEnumPropertyBinding::Load(FEnumProperty& Dst, FStructLoadView Src, ECustom
 
 ////////////////////////////////////////////////////////////////////////////////////////////////
 
-UStructBinding::UStructBinding(TCustomSpecifier<Ids, NumMembers>& Spec)
+UStructBinding::UStructBinding(TPropertySpecifier<NumMembers>& Spec)
 {
 	Spec.Members[M::SuperStruct] = GUE.Structs.ObjectPtr;
 	// Children is an array of UFields, these are always UFunctions
@@ -516,7 +516,7 @@ void UStructBinding::Load(UStruct& Dst, FStructLoadView Src, ECustomLoadMethod M
 	check(!Members.HasMore());
 }
 
-UClassBinding::UClassBinding(TCustomSpecifier<Ids, NumMembers>& Spec)
+UClassBinding::UClassBinding(TPropertySpecifier<NumMembers>& Spec)
 : ClassFlagsId(GetEnumId<Ids, EClassFlags>())
 , ImplementedInterfaceId(GetStructDualId<Ids, FImplementedInterface>())
 {
@@ -621,7 +621,7 @@ void UClassBinding::Load(UClass& Dst, FStructLoadView Src, ECustomLoadMethod Met
 	Dst.bCooked = false;
 }
 
-UFunctionBinding::UFunctionBinding(TCustomSpecifier<Ids, NumMembers>& Spec)
+UFunctionBinding::UFunctionBinding(TPropertySpecifier<NumMembers>& Spec)
 : FunctionFlagsId(GetEnumId<Ids, EFunctionFlags>())
 {
 	Spec.Super = GUE.Structs.Struct;
@@ -649,7 +649,7 @@ void UFunctionBinding::Load(UFunction& Dst, FStructLoadView Src, ECustomLoadMeth
 	check(!Members.HasMore());
 }
 
-UScriptStructBinding::UScriptStructBinding(TCustomSpecifier<Ids, NumMembers>& Spec)
+UScriptStructBinding::UScriptStructBinding(TPropertySpecifier<NumMembers>& Spec)
 : StructFlagsId(GetEnumId<Ids, EStructFlags>())
 {
 	Spec.Super = GUE.Structs.Struct;
